Attempt to fix XB1 compile error. Compare against nullptr rather than NULL.

diff --git a/Engine/Source/Developer/Settings/Private/SettingsCategory.h b/Engine/Source/Developer/Settings/Private/SettingsCategory.h
--- a/Engine/Source/Developer/Settings/Private/SettingsCategory.h
+++ b/Engine/Source/Developer/Settings/Private/SettingsCategory.h
@@ -69,7 +69,7 @@ public:
 	{
 		TSharedPtr<FSettingsSection>& Section = Sections.FindOrAdd(SectionName);
 
-		if (!Section.IsValid() || (Section->GetSettingsObject() != NULL) || (Section->GetCustomWidget().Pin() != CustomWidget))
+		if (!Section.IsValid() || (Section->GetSettingsObject() != nullptr) || (Section->GetCustomWidget().Pin() != CustomWidget))
 		{
 			Section = MakeShareable(new FSettingsSection(AsShared(), SectionName, InDisplayName, InDescription, Delegates, CustomWidget));
 		}
