Settings API: changed all occurences of NULL to nullptr

diff --git a/Engine/Source/Developer/Settings/Private/SettingsContainer.h b/Engine/Source/Developer/Settings/Private/SettingsContainer.h
--- a/Engine/Source/Developer/Settings/Private/SettingsContainer.h
+++ b/Engine/Source/Developer/Settings/Private/SettingsContainer.h
@@ -38,7 +38,7 @@ public:
 	 * @param SettingsObject - The object that holds the section's settings.
 	 * @param Delegates - The section's optional callback delegates.
 	 *
-	 * @return The added settings section, or NULL if the category does not exist.
+	 * @return The added settings section, or nullptr if the category does not exist.
 	 */
 	ISettingsSectionPtr AddSection( const FName& CategoryName, const FName& SectionName, const FText& InDisplayName, const FText& InDescription, const TWeakObjectPtr<UObject>& SettingsObject, const FSettingsSectionDelegates& Delegates )
 	{
@@ -46,7 +46,7 @@ public:
 
 		if (!Category.IsValid())
 		{
-			return NULL;
+			return nullptr;
 		}
 
 		ISettingsSectionRef Section = Category->AddSection(SectionName, InDisplayName, InDescription, SettingsObject, Delegates);
@@ -67,7 +67,7 @@ public:
 	 * @param CustomWidget - A custom settings widget.
 	 * @param Delegates - The section's optional callback delegates.
 	 *
-	 * @return The added settings section, or NULL if the category does not exist.
+	 * @return The added settings section, or nullptr if the category does not exist.
 	 */
 	ISettingsSectionPtr AddSection( const FName& CategoryName, const FName& SectionName, const FText& InDisplayName, const FText& InDescription, const TSharedRef<SWidget>& CustomWidget, const FSettingsSectionDelegates& Delegates )
 	{
@@ -75,7 +75,7 @@ public:
 
 		if (!Category.IsValid())
 		{
-			return NULL;
+			return nullptr;
 		}
 
 		ISettingsSectionRef Section = Category->AddSection(SectionName, InDisplayName, InDescription, CustomWidget, Delegates);
diff --git a/Engine/Source/Developer/Settings/Private/SettingsModule.cpp b/Engine/Source/Developer/Settings/Private/SettingsModule.cpp
--- a/Engine/Source/Developer/Settings/Private/SettingsModule.cpp
+++ b/Engine/Source/Developer/Settings/Private/SettingsModule.cpp
@@ -49,7 +49,7 @@ public:
 	{
 		ISettingsViewer** Viewer = ContainerNamesToViewers.Find(ContainerName);
 
-		if (Viewer != NULL)
+		if (Viewer != nullptr)
 		{
 			(*Viewer)->ShowSettings(CategoryName, SectionName);
 		}
diff --git a/Engine/Source/Developer/Settings/Private/SettingsSection.h b/Engine/Source/Developer/Settings/Private/SettingsSection.h
--- a/Engine/Source/Developer/Settings/Private/SettingsSection.h
+++ b/Engine/Source/Developer/Settings/Private/SettingsSection.h
@@ -204,9 +204,9 @@ public:
 			GConfig->EmptySection(*SettingsObject->GetClass()->GetPathName(), ConfigName);
 			GConfig->Flush(false);
 
-			FConfigCacheIni::LoadGlobalIniFile(ConfigName, *FPaths::GetBaseFilename(ConfigName), NULL, NULL, true);
+			FConfigCacheIni::LoadGlobalIniFile(ConfigName, *FPaths::GetBaseFilename(ConfigName), nullptr, nullptr, true);
 
-			SettingsObject->ReloadConfig(NULL, NULL, UE4::LCPF_PropagateToInstances|UE4::LCPF_PropagateToChildDefaultObjects);
+			SettingsObject->ReloadConfig(nullptr, nullptr, UE4::LCPF_PropagateToInstances|UE4::LCPF_PropagateToChildDefaultObjects);
 
 			return true;
 		}
@@ -246,7 +246,7 @@ public:
 		if (SettingsObject.IsValid())
 		{
 			SettingsObject->UpdateDefaultConfigFile();
-			SettingsObject->ReloadConfig(NULL, NULL, UE4::LCPF_PropagateToInstances);
+			SettingsObject->ReloadConfig(nullptr, nullptr, UE4::LCPF_PropagateToInstances);
 
 			return true;			
 		}
diff --git a/Engine/Source/Developer/Settings/Public/Interfaces/ISettingsCategory.h b/Engine/Source/Developer/Settings/Public/Interfaces/ISettingsCategory.h
--- a/Engine/Source/Developer/Settings/Public/Interfaces/ISettingsCategory.h
+++ b/Engine/Source/Developer/Settings/Public/Interfaces/ISettingsCategory.h
@@ -48,7 +48,7 @@ public:
 	 *
 	 * @param SectionName - The name of the section to get.
 	 *
-	 * @return The settings section, or NULL if it doesn't exist.
+	 * @return The settings section, or nullptr if it doesn't exist.
 	 */
 	virtual ISettingsSectionPtr GetSection( const FName& SectionName ) const = 0;
 
diff --git a/Engine/Source/Developer/Settings/Public/Interfaces/ISettingsContainer.h b/Engine/Source/Developer/Settings/Public/Interfaces/ISettingsContainer.h
--- a/Engine/Source/Developer/Settings/Public/Interfaces/ISettingsContainer.h
+++ b/Engine/Source/Developer/Settings/Public/Interfaces/ISettingsContainer.h
@@ -75,7 +75,7 @@ public:
 	/**
 	 * Gets the category with the specified name.
 	 *
-	 * @return The category, or NULL if it doesn't exist.
+	 * @return The category, or nullptr if it doesn't exist.
 	 */
 	virtual ISettingsCategoryPtr GetCategory( const FName& CategoryName ) const = 0;
 
diff --git a/Engine/Source/Developer/Settings/Public/Interfaces/ISettingsModule.h b/Engine/Source/Developer/Settings/Public/Interfaces/ISettingsModule.h
--- a/Engine/Source/Developer/Settings/Public/Interfaces/ISettingsModule.h
+++ b/Engine/Source/Developer/Settings/Public/Interfaces/ISettingsModule.h
@@ -27,7 +27,7 @@ public:
 	 *
 	 * @param ContainerName - The name of the container to get.
 	 *
-	 * @return The settings container, or NULL if it doesn't exist.
+	 * @return The settings container, or nullptr if it doesn't exist.
 	 */
 	virtual ISettingsContainerPtr GetContainer( const FName& ContainerName ) = 0;
 
@@ -44,7 +44,7 @@ public:
 	 * @param SettingsObject - The object that holds the section's settings.
 	 * @param Delegates - The section's optional callback delegates.
 	 *
-	 * @return The added settings section, or NULL if the category does not exist.
+	 * @return The added settings section, or nullptr if the category does not exist.
 	 */
 	virtual ISettingsSectionPtr RegisterSettings( const FName& ContainerName, const FName& CategoryName, const FName& SectionName, const FText& DisplayName, const FText& Description, const TWeakObjectPtr<UObject>& SettingsObject, const FSettingsSectionDelegates& Delegates ) = 0;
 
@@ -61,7 +61,7 @@ public:
 	 * @param CustomWidget - A custom settings widget.
 	 * @param Delegates - The section's optional callback delegates.
 	 *
-	 * @return The added settings section, or NULL if the category does not exist.
+	 * @return The added settings section, or nullptr if the category does not exist.
 	 */
 	virtual ISettingsSectionPtr RegisterSettings( const FName& ContainerName, const FName& CategoryName, const FName& SectionName, const FText& DisplayName, const FText& Description, const TSharedRef<SWidget>& CustomWidget, const FSettingsSectionDelegates& Delegates ) = 0;
 
@@ -111,7 +111,7 @@ public:
 	 * @param Description - The section's localized description text.
 	 * @param SettingsObject - The object that holds the section's settings.
 	 *
-	 * @return The added settings section, or NULL if the category does not exist.
+	 * @return The added settings section, or nullptr if the category does not exist.
 	 */
 	ISettingsSectionPtr RegisterSettings( const FName& ContainerName, const FName& CategoryName, const FName& SectionName, const FText& DisplayName, const FText& Description, const TWeakObjectPtr<UObject>& SettingsObject )
 	{
@@ -130,7 +130,7 @@ public:
 	 * @param Description - The section's localized description text.
 	 * @param CustomWidget - A custom settings widget.
 	 *
-	 * @return The added settings section, or NULL if the category does not exist.
+	 * @return The added settings section, or nullptr if the category does not exist.
 	 */
 	ISettingsSectionPtr RegisterSettings( const FName& ContainerName, const FName& CategoryName, const FName& SectionName, const FText& DisplayName, const FText& Description, const TSharedRef<SWidget>& CustomWidget )
 	{
@@ -144,7 +144,7 @@ public:
 	 *
 	 * @todo gmp: better implementation using dependency injection / abstract factory.
 	 *
-	 * @return A pointer to the Settings module, or NULL if the module isn't loaded.
+	 * @return A pointer to the Settings module, or nullptr if the module isn't loaded.
 	 */
 	static ISettingsModule* Get( )
 	{
