SettingsEditor: disabled the 'Set as Default' button for defaultconfig settings objects (which are always default already)

diff --git a/Engine/Source/Developer/Settings/Private/SettingsSection.h b/Engine/Source/Developer/Settings/Private/SettingsSection.h
--- a/Engine/Source/Developer/Settings/Private/SettingsSection.h
+++ b/Engine/Source/Developer/Settings/Private/SettingsSection.h
@@ -103,7 +103,7 @@ public:
 
 	virtual bool CanSaveDefaults( ) const OVERRIDE
 	{
-		return (SaveDefaultsDelegate.IsBound() || (SettingsObject.IsValid() && SettingsObject->GetClass()->HasAnyClassFlags(CLASS_Config)));
+		return (SaveDefaultsDelegate.IsBound() || (SettingsObject.IsValid() && SettingsObject->GetClass()->HasAnyClassFlags(CLASS_Config) && !SettingsObject->GetClass()->HasAnyClassFlags(CLASS_DefaultConfig)));
 	}
 
 	virtual bool Export( const FString& Filename ) OVERRIDE
