Settings: Documentation cleanup pass

diff --git a/Engine/Source/Developer/Settings/Private/SettingsCategory.h b/Engine/Source/Developer/Settings/Private/SettingsCategory.h
--- a/Engine/Source/Developer/Settings/Private/SettingsCategory.h
+++ b/Engine/Source/Developer/Settings/Private/SettingsCategory.h
@@ -19,7 +19,7 @@ public:
 	/**
 	 * Creates and initializes a new instance.
 	 *
-	 * @param InName - The category's name.
+	 * @param InName The category's name.
 	 */
 	FSettingsCategory( const FName& InName )
 		: Name(InName)
@@ -32,12 +32,11 @@ public:
 	 *
 	 * If a section with the specified settings objects already exists, the existing section will be returned.
 	 *
-	 * @param SectionName - The name of the settings section to add.
-	 * @param DisplayName - The section's localized display name.
-	 * @param Description - The section's localized description text.
-	 * @param SettingsObject - The object that holds the section's settings.
-	 * @param Delegates - The section's optional callback delegates.
-	 *
+	 * @param SectionName The name of the settings section to add.
+	 * @param DisplayName The section's localized display name.
+	 * @param Description The section's localized description text.
+	 * @param SettingsObject The object that holds the section's settings.
+	 * @param Delegates The section's optional callback delegates.
 	 * @return The added settings section.
 	 */
 	ISettingsSectionRef AddSection( const FName& SectionName, const FText& InDisplayName, const FText& InDescription, const TWeakObjectPtr<UObject>& SettingsObject, const FSettingsSectionDelegates& Delegates )
@@ -57,12 +56,11 @@ public:
 	 *
 	 * If a section with the specified settings objects already exists, the existing section will be returned.
 	 *
-	 * @param SectionName - The name of the settings section to add.
-	 * @param DisplayName - The section's localized display name.
-	 * @param Description - The section's localized description text.
-	 * @param CustomWidget - A custom settings widget.
-	 * @param Delegates - The section's optional callback delegates.
-	 *
+	 * @param SectionName The name of the settings section to add.
+	 * @param DisplayName The section's localized display name.
+	 * @param Description The section's localized description text.
+	 * @param CustomWidget A custom settings widget.
+	 * @param Delegates The section's optional callback delegates.
 	 * @return The added settings section.
 	 */
 	ISettingsSectionRef AddSection( const FName& SectionName, const FText& InDisplayName, const FText& InDescription, const TSharedRef<SWidget>& CustomWidget, const FSettingsSectionDelegates& Delegates )
@@ -80,10 +78,9 @@ public:
 	/**
 	 * Updates the details of this settings category.
 	 *
-	 * @param InDisplayName - The category's localized display name.
-	 * @param InDescription - The category's localized description text.
-	 * @param InIconName - The name of the category's icon.
-	 *
+	 * @param InDisplayName The category's localized display name.
+	 * @param InDescription The category's localized description text.
+	 * @param InIconName The name of the category's icon.
 	 * @return The category.
 	 */
 	void Describe( const FText& InDisplayName, const FText& InDescription, const FName& InIconName )
@@ -96,7 +93,7 @@ public:
 	/**
 	 * Removes a settings section.
 	 *
-	 * @param SectionName - The name of the section to remove.
+	 * @param SectionName The name of the section to remove.
 	 */
 	void RemoveSection( const FName& SectionName )
 	{
diff --git a/Engine/Source/Developer/Settings/Private/SettingsContainer.h b/Engine/Source/Developer/Settings/Private/SettingsContainer.h
--- a/Engine/Source/Developer/Settings/Private/SettingsContainer.h
+++ b/Engine/Source/Developer/Settings/Private/SettingsContainer.h
@@ -31,13 +31,12 @@ public:
 	 *
 	 * If a section with the specified settings objects already exists, the existing section will be replaced.
 	 *
-	 * @param CategoryName - The name of the category to add the section to.
-	 * @param SectionName - The name of the settings section to add.
-	 * @param DisplayName - The section's localized display name.
-	 * @param Description - The section's localized description text.
-	 * @param SettingsObject - The object that holds the section's settings.
-	 * @param Delegates - The section's optional callback delegates.
-	 *
+	 * @param CategoryName The name of the category to add the section to.
+	 * @param SectionName The name of the settings section to add.
+	 * @param DisplayName The section's localized display name.
+	 * @param Description The section's localized description text.
+	 * @param SettingsObject The object that holds the section's settings.
+	 * @param Delegates The section's optional callback delegates.
 	 * @return The added settings section, or nullptr if the category does not exist.
 	 */
 	ISettingsSectionPtr AddSection( const FName& CategoryName, const FName& SectionName, const FText& InDisplayName, const FText& InDescription, const TWeakObjectPtr<UObject>& SettingsObject, const FSettingsSectionDelegates& Delegates )
@@ -60,13 +59,12 @@ public:
 	 *
 	 * If a section with the specified settings objects already exists, the existing section will be replaced.
 	 *
-	 * @param CategoryName - The name of the category to add the section to.
-	 * @param SectionName - The name of the settings section to add.
-	 * @param DisplayName - The section's localized display name.
-	 * @param Description - The section's localized description text.
-	 * @param CustomWidget - A custom settings widget.
-	 * @param Delegates - The section's optional callback delegates.
-	 *
+	 * @param CategoryName The name of the category to add the section to.
+	 * @param SectionName The name of the settings section to add.
+	 * @param DisplayName The section's localized display name.
+	 * @param Description The section's localized description text.
+	 * @param CustomWidget A custom settings widget.
+	 * @param Delegates The section's optional callback delegates.
 	 * @return The added settings section, or nullptr if the category does not exist.
 	 */
 	ISettingsSectionPtr AddSection( const FName& CategoryName, const FName& SectionName, const FText& InDisplayName, const FText& InDescription, const TSharedRef<SWidget>& CustomWidget, const FSettingsSectionDelegates& Delegates )
@@ -87,8 +85,8 @@ public:
 	/**
 	 * Removes a settings section.
 	 *
-	 * @param CategoryName - The name of the category that contains the section.
-	 * @param SectionName - The name of the section to remove.
+	 * @param CategoryName The name of the category that contains the section.
+	 * @param SectionName The name of the section to remove.
 	 */
 	void RemoveSection( const FName& CategoryName, const FName& SectionName )
 	{
@@ -183,7 +181,7 @@ public:
 private:
 
 	// Holds the collection of setting categories
-	TMap<FName, TSharedPtr<FSettingsCategory> > Categories;
+	TMap<FName, TSharedPtr<FSettingsCategory>> Categories;
 
 	// Holds the container's description text.
 	FText Description;
diff --git a/Engine/Source/Developer/Settings/Private/SettingsSection.h b/Engine/Source/Developer/Settings/Private/SettingsSection.h
--- a/Engine/Source/Developer/Settings/Private/SettingsSection.h
+++ b/Engine/Source/Developer/Settings/Private/SettingsSection.h
@@ -18,12 +18,12 @@ public:
 	/**
 	 * Creates and initializes a new settings section from the given settings object.
 	 *
-	 * @param InCategory - The settings category that owns this section.
-	 * @param InName - The setting section's name.
-	 * @param InDisplayName - The section's localized display name.
-	 * @param InDescription - The section's localized description text.
-	 * @param InDelegates - The section's optional callback delegates.
-	 * @param InSettingsObject - The object that holds the settings for this section.
+	 * @param InCategory The settings category that owns this section.
+	 * @param InName The setting section's name.
+	 * @param InDisplayName The section's localized display name.
+	 * @param InDescription The section's localized description text.
+	 * @param InDelegates The section's optional callback delegates.
+	 * @param InSettingsObject The object that holds the settings for this section.
 	 */
 	FSettingsSection( const ISettingsCategoryRef& InCategory, const FName& InName, const FText& InDisplayName, const FText& InDescription, const FSettingsSectionDelegates& InDelegates, const TWeakObjectPtr<UObject>& InSettingsObject )
 		: Category(InCategory)
@@ -44,12 +44,12 @@ public:
 	/**
 	 * Creates and initializes a new settings section from the given custom settings widget.
 	 *
-	 * @param InCategory - The settings category that owns this section.
-	 * @param InName - The setting section's name.
-	 * @param InDisplayName - The section's localized display name.
-	 * @param InDescription - The section's localized description text.
-	 * @param InDelegates - The section's optional callback delegates.
-	 * @param InCustomWidget - A custom settings widget.
+	 * @param InCategory The settings category that owns this section.
+	 * @param InName The setting section's name.
+	 * @param InDisplayName The section's localized display name.
+	 * @param InDescription The section's localized description text.
+	 * @param InDelegates The section's optional callback delegates.
+	 * @param InCustomWidget A custom settings widget.
 	 */
 	FSettingsSection( const ISettingsCategoryRef& InCategory, const FName& InName, const FText& InDisplayName, const FText& InDescription, const FSettingsSectionDelegates& InDelegates, const TSharedRef<SWidget>& InCustomWidget )
 		: Category(InCategory)
diff --git a/Engine/Source/Developer/Settings/Public/Interfaces/ISettingsCategory.h b/Engine/Source/Developer/Settings/Public/Interfaces/ISettingsCategory.h
--- a/Engine/Source/Developer/Settings/Public/Interfaces/ISettingsCategory.h
+++ b/Engine/Source/Developer/Settings/Public/Interfaces/ISettingsCategory.h
@@ -7,14 +7,10 @@
 #pragma once
 
 
-/**
- * Type definition for shared pointers to instances of ISettingsCategory.
- */
+/** Type definition for shared pointers to instances of ISettingsCategory. */
 typedef TSharedPtr<class ISettingsCategory> ISettingsCategoryPtr;
 
-/**
- * Type definition for shared references to instances of ISettingsCategory.
- */
+/** Type definition for shared references to instances of ISettingsCategory. */
 typedef TSharedRef<class ISettingsCategory> ISettingsCategoryRef;
 
 
@@ -46,8 +42,7 @@ public:
 	/**
 	 * Gets the settings section with the specified name.
 	 *
-	 * @param SectionName - The name of the section to get.
-	 *
+	 * @param SectionName The name of the section to get.
 	 * @return The settings section, or nullptr if it doesn't exist.
 	 */
 	virtual ISettingsSectionPtr GetSection( const FName& SectionName ) const = 0;
@@ -55,8 +50,7 @@ public:
 	/**
 	 * Gets the setting sections contained in this category.
 	 *
-	 * @param OutSections - Will hold the collection of sections.
-	 *
+	 * @param OutSections Will hold the collection of sections.
 	 * @return The number of sections returned.
 	 */
 	virtual int32 GetSections( TArray<ISettingsSectionPtr>& OutSections ) const = 0;
diff --git a/Engine/Source/Developer/Settings/Public/Interfaces/ISettingsContainer.h b/Engine/Source/Developer/Settings/Public/Interfaces/ISettingsContainer.h
--- a/Engine/Source/Developer/Settings/Public/Interfaces/ISettingsContainer.h
+++ b/Engine/Source/Developer/Settings/Public/Interfaces/ISettingsContainer.h
@@ -22,14 +22,10 @@ DECLARE_MULTICAST_DELEGATE_OneParam(FOnSettingsContainerCategoryModified, const
 DECLARE_MULTICAST_DELEGATE_OneParam(FOnSettingsContainerSectionRemoved, const ISettingsSectionRef&)
 
 
-/**
- * Type definition for shared pointers to instances of ISettingsContainer.
- */
+/** Type definition for shared pointers to instances of ISettingsContainer. */
 typedef TSharedPtr<class ISettingsContainer> ISettingsContainerPtr;
 
-/**
- * Type definition for shared references to instances of ISettingsContainer.
- */
+/** Type definition for shared references to instances of ISettingsContainer. */
 typedef TSharedRef<class ISettingsContainer> ISettingsContainerRef;
 
 
@@ -47,27 +43,26 @@ public:
 	/**
 	 * Updates the details of this settings container.
 	 *
-	 * @param InDisplayName - The container's localized display name.
-	 * @param InDescription - The container's localized description text.
-	 * @param InIconName - The name of the container's icon.
+	 * @param InDisplayName The container's localized display name.
+	 * @param InDescription The container's localized description text.
+	 * @param InIconName The name of the container's icon.
 	 */
 	virtual void Describe( const FText& InDisplayName, const FText& InDescription, const FName& InIconName ) = 0;
 
 	/**
 	 * Updates the details of the specified settings category.
 	 *
-	 * @param CategoryName - The name of the category to update.
-	 * @param DisplayName - The category's localized display name.
-	 * @param Description - The category's localized description text.
-	 * @param IconName - The name of the category's icon.
+	 * @param CategoryName The name of the category to update.
+	 * @param DisplayName The category's localized display name.
+	 * @param Description The category's localized description text.
+	 * @param IconName The name of the category's icon.
 	 */
 	virtual void DescribeCategory( const FName& CategoryName, const FText& DisplayName, const FText& Description, const FName& IconName ) = 0;
 
 	/**
 	 * Gets the setting categories.
 	 *
-	 * @param OutCategories - Will hold the collection of categories.
-	 *
+	 * @param OutCategories Will hold the collection of categories.
 	 * @return The number of categories returned.
 	 */
 	virtual int32 GetCategories( TArray<ISettingsCategoryPtr>& OutCategories ) const = 0;
diff --git a/Engine/Source/Developer/Settings/Public/Interfaces/ISettingsModule.h b/Engine/Source/Developer/Settings/Public/Interfaces/ISettingsModule.h
--- a/Engine/Source/Developer/Settings/Public/Interfaces/ISettingsModule.h
+++ b/Engine/Source/Developer/Settings/Public/Interfaces/ISettingsModule.h
@@ -18,15 +18,14 @@ public:
 	/**
 	 * Gets the names of all known setting containers.
 	 *
-	 * @param OutNames - Will contain the collection of names.
+	 * @param OutNames Will contain the collection of names.
 	 */
 	virtual void GetContainerNames( TArray<FName>& OutNames ) const = 0;
 
 	/**
 	 * Gets the global settings container with the specified name.
 	 *
-	 * @param ContainerName - The name of the container to get.
-	 *
+	 * @param ContainerName The name of the container to get.
 	 * @return The settings container, or nullptr if it doesn't exist.
 	 */
 	virtual ISettingsContainerPtr GetContainer( const FName& ContainerName ) = 0;
@@ -36,14 +35,13 @@ public:
 	 *
 	 * If a section with the specified settings objects already exists, the existing section will be replaced.
 	 *
-	 * @param ContainerName - The name of the container that will contain the settings.
-	 * @param CategoryName - The name of the category within the container.
-	 * @param SectionName - The name of the section within the category.
-	 * @param DisplayName - The section's localized display name.
-	 * @param Description - The section's localized description text.
-	 * @param SettingsObject - The object that holds the section's settings.
-	 * @param Delegates - The section's optional callback delegates.
-	 *
+	 * @param ContainerName The name of the container that will contain the settings.
+	 * @param CategoryName The name of the category within the container.
+	 * @param SectionName The name of the section within the category.
+	 * @param DisplayName The section's localized display name.
+	 * @param Description The section's localized description text.
+	 * @param SettingsObject The object that holds the section's settings.
+	 * @param Delegates The section's optional callback delegates.
 	 * @return The added settings section, or nullptr if the category does not exist.
 	 */
 	virtual ISettingsSectionPtr RegisterSettings( const FName& ContainerName, const FName& CategoryName, const FName& SectionName, const FText& DisplayName, const FText& Description, const TWeakObjectPtr<UObject>& SettingsObject, const FSettingsSectionDelegates& Delegates ) = 0;
@@ -53,14 +51,13 @@ public:
 	 *
 	 * If a section with the specified settings objects already exists, the existing section will be replaced.
 	 *
-	 * @param ContainerName - The name of the container that will contain the settings.
-	 * @param CategoryName - The name of the category within the container.
-	 * @param SectionName - The name of the section within the category.
-	 * @param DisplayName - The section's localized display name.
-	 * @param Description - The section's localized description text.
-	 * @param CustomWidget - A custom settings widget.
-	 * @param Delegates - The section's optional callback delegates.
-	 *
+	 * @param ContainerName The name of the container that will contain the settings.
+	 * @param CategoryName The name of the category within the container.
+	 * @param SectionName The name of the section within the category.
+	 * @param DisplayName The section's localized display name.
+	 * @param Description The section's localized description text.
+	 * @param CustomWidget A custom settings widget.
+	 * @param Delegates The section's optional callback delegates.
 	 * @return The added settings section, or nullptr if the category does not exist.
 	 */
 	virtual ISettingsSectionPtr RegisterSettings( const FName& ContainerName, const FName& CategoryName, const FName& SectionName, const FText& DisplayName, const FText& Description, const TSharedRef<SWidget>& CustomWidget, const FSettingsSectionDelegates& Delegates ) = 0;
@@ -68,33 +65,33 @@ public:
 	/**
 	 * Registers a viewer for the specified settings container.
 	 *
-	 * @param ContainerName - The name of the settings container to register a viewer for.
-	 * @param SettingsViewer - The viewer to register.
+	 * @param ContainerName The name of the settings container to register a viewer for.
+	 * @param SettingsViewer The viewer to register.
 	 */
 	virtual void RegisterViewer( const FName& ContainerName, ISettingsViewer& SettingsViewer ) = 0;
 
 	/**
 	 * Shows the settings viewer for the specified settings container.
 	 *
-	 * @param ContainerName - The name of the section's container.
-	 * @param CategoryName - The name of the section's category.
-	 * @param SectionName - The name of the section to show.
+	 * @param ContainerName The name of the section's container.
+	 * @param CategoryName The name of the section's category.
+	 * @param SectionName The name of the section to show.
 	 */
 	virtual void ShowViewer( const FName& ContainerName, const FName& CategoryName, const FName& SectionName ) = 0;
 
 	/**
 	 * Removes a settings section from the specified settings container.
 	 *
-	 * @param ContainerName - The name of the container that to remove the settings from.
-	 * @param CategoryName - The name of the category within the container.
-	 * @param SectionName - The name of the section within the category.
+	 * @param ContainerName The name of the container that to remove the settings from.
+	 * @param CategoryName The name of the category within the container.
+	 * @param SectionName The name of the section within the category.
 	 */
 	virtual void UnregisterSettings( const FName& ContainerName, const FName& CategoryName, const FName& SectionName ) = 0;
 
 	/**
 	 * Unregisters the currently assigned viewer for the specified settings container.
 	 *
-	 * @param ContainerName - The name of the settings container to unregister the viewer for.
+	 * @param ContainerName The name of the settings container to unregister the viewer for.
 	 */
 	virtual void UnregisterViewer( const FName& ContainerName ) = 0;
 
@@ -105,12 +102,11 @@ public:
 	 *
 	 * If a section with the specified settings objects already exists, the existing section will be replaced.
 	 *
-	 * @param CategoryName - The name of the category to add the section to.
-	 * @param SectionName - The name of the settings section to add.
-	 * @param DisplayName - The section's localized display name.
-	 * @param Description - The section's localized description text.
-	 * @param SettingsObject - The object that holds the section's settings.
-	 *
+	 * @param CategoryName The name of the category to add the section to.
+	 * @param SectionName The name of the settings section to add.
+	 * @param DisplayName The section's localized display name.
+	 * @param Description The section's localized description text.
+	 * @param SettingsObject The object that holds the section's settings.
 	 * @return The added settings section, or nullptr if the category does not exist.
 	 */
 	ISettingsSectionPtr RegisterSettings( const FName& ContainerName, const FName& CategoryName, const FName& SectionName, const FText& DisplayName, const FText& Description, const TWeakObjectPtr<UObject>& SettingsObject )
@@ -123,13 +119,12 @@ public:
 	 *
 	 * If a section with the specified settings objects already exists, the existing section will be replaced.
 	 *
-	 * @param ContainerName - The name of the container that will contain the settings.
-	 * @param CategoryName - The name of the category within the container.
-	 * @param SectionName - The name of the section within the category.
-	 * @param DisplayName - The section's localized display name.
-	 * @param Description - The section's localized description text.
-	 * @param CustomWidget - A custom settings widget.
-	 *
+	 * @param ContainerName The name of the container that will contain the settings.
+	 * @param CategoryName The name of the category within the container.
+	 * @param SectionName The name of the section within the category.
+	 * @param DisplayName The section's localized display name.
+	 * @param Description The section's localized description text.
+	 * @param CustomWidget A custom settings widget.
 	 * @return The added settings section, or nullptr if the category does not exist.
 	 */
 	ISettingsSectionPtr RegisterSettings( const FName& ContainerName, const FName& CategoryName, const FName& SectionName, const FText& DisplayName, const FText& Description, const TSharedRef<SWidget>& CustomWidget )
diff --git a/Engine/Source/Developer/Settings/Public/Interfaces/ISettingsSection.h b/Engine/Source/Developer/Settings/Public/Interfaces/ISettingsSection.h
--- a/Engine/Source/Developer/Settings/Public/Interfaces/ISettingsSection.h
+++ b/Engine/Source/Developer/Settings/Public/Interfaces/ISettingsSection.h
@@ -113,19 +113,13 @@ struct FSettingsSectionDelegates
 };
 
 
-/**
- * Type definition for weak pointers to instances of ISettingsCategory.
- */
+/** Type definition for weak pointers to instances of ISettingsCategory. */
 typedef TWeakPtr<class ISettingsCategory> ISettingsCategoryWeakPtr;
 
-/**
- * Type definition for shared pointers to instances of ISettingsSection.
- */
+/** Type definition for shared pointers to instances of ISettingsSection. */
 typedef TSharedPtr<class ISettingsSection> ISettingsSectionPtr;
 
-/**
- * Type definition for shared references to instances of ISettingsSection.
- */
+/** Type definition for shared references to instances of ISettingsSection. */
 typedef TSharedRef<class ISettingsSection> ISettingsSectionRef;
 
 
@@ -189,8 +183,7 @@ public:
 	/**
 	 * Exports the settings in this section to the specified file.
 	 *
-	 * @param Filename - The path to the file.
-	 *
+	 * @param Filename The path to the file.
 	 * @return true if the settings were exported, false otherwise.
 	 */
 	virtual bool Export( const FString& Filename ) = 0;
@@ -254,8 +247,7 @@ public:
 	/**
 	 * Imports the settings in this section from the specified file.
 	 *
-	 * @param Filename - The path to the file.
-	 *
+	 * @param Filename The path to the file.
 	 * @return true if the settings were imported, false otherwise.
 	 */
 	virtual bool Import( const FString& Filename ) = 0;
diff --git a/Engine/Source/Developer/Settings/Public/Interfaces/ISettingsViewer.h b/Engine/Source/Developer/Settings/Public/Interfaces/ISettingsViewer.h
--- a/Engine/Source/Developer/Settings/Public/Interfaces/ISettingsViewer.h
+++ b/Engine/Source/Developer/Settings/Public/Interfaces/ISettingsViewer.h
@@ -17,8 +17,8 @@ public:
 	/**
 	 * Shows the settings tab that belongs to this viewer with the specified settings section pre-selected.
 	 *
-	 * @param CategoryName - The name of the section's category.
-	 * @param SectionName - The name of the section to select.
+	 * @param CategoryName The name of the section's category.
+	 * @param SectionName The name of the section to select.
 	 */
 	virtual void ShowSettings( const FName& CategoryName, const FName& SectionName ) = 0;
 
