Settings: Removed category icons and the corresponding APIs

The placeholder category icons were never replaced with actual icons, and hence they do not fulfill any purpose. Icons as a visual language for categories also conflicts with their usage as interactive UI elements, such as buttons and toolbar options.

diff --git a/Engine/Source/Developer/Settings/Private/SettingsCategory.h b/Engine/Source/Developer/Settings/Private/SettingsCategory.h
--- a/Engine/Source/Developer/Settings/Private/SettingsCategory.h
+++ b/Engine/Source/Developer/Settings/Private/SettingsCategory.h
@@ -79,11 +79,10 @@ public:
 	 * @param InIconName The name of the category's icon.
 	 * @return The category.
 	 */
-	void Describe( const FText& InDisplayName, const FText& InDescription, const FName& InIconName )
+	void Describe( const FText& InDisplayName, const FText& InDescription )
 	{
 		Description = InDescription;
 		DisplayName = InDisplayName;
-		IconName = InIconName;
 	}
 
 	/**
@@ -110,11 +109,6 @@ public:
 		return DisplayName;
 	}
 
-	virtual const FName& GetIconName( ) const override
-	{
-		return IconName;
-	}
-
 	virtual const FName& GetName( ) const override
 	{
 		return Name;
@@ -148,9 +142,6 @@ private:
 	// Holds the collection of setting sections.
 	TMap<FName, TSharedPtr<FSettingsSection> > Sections;
 
-	// Holds the name of the category's icon.
-	FName IconName;
-
 	// Holds the category's name.
 	FName Name;
 };
diff --git a/Engine/Source/Developer/Settings/Private/SettingsContainer.h b/Engine/Source/Developer/Settings/Private/SettingsContainer.h
--- a/Engine/Source/Developer/Settings/Private/SettingsContainer.h
+++ b/Engine/Source/Developer/Settings/Private/SettingsContainer.h
@@ -112,7 +112,7 @@ public:
 		IconName = InIconName;
 	}
 
-	virtual void DescribeCategory( const FName& CategoryName, const FText& InDisplayName, const FText& InDescription, const FName& InIconName ) override
+	virtual void DescribeCategory( const FName& CategoryName, const FText& InDisplayName, const FText& InDescription ) override
 	{
 		TSharedPtr<FSettingsCategory>& Category = Categories.FindOrAdd(CategoryName);
 
@@ -121,7 +121,7 @@ public:
 			Category = MakeShareable(new FSettingsCategory(CategoryName));
 		}
 
-		Category->Describe(InDisplayName, InDescription, InIconName);
+		Category->Describe(InDisplayName, InDescription);
 		CategoryModifiedDelegate.Broadcast(CategoryName);
 	}
 
diff --git a/Engine/Source/Developer/Settings/Private/SettingsModule.cpp b/Engine/Source/Developer/Settings/Private/SettingsModule.cpp
--- a/Engine/Source/Developer/Settings/Private/SettingsModule.cpp
+++ b/Engine/Source/Developer/Settings/Private/SettingsModule.cpp
@@ -75,18 +75,18 @@ public:
 		// @todo gmp: move this into the modules that own these setting categories
 		TSharedRef<FSettingsContainer> EditorSettingsContainer = FindOrAddContainer("Editor");
 		EditorSettingsContainer->Describe(LOCTEXT("EditorPreferencesSubMenuLabel", "Editor Preferences"), LOCTEXT("EditorPreferencesSubMenuToolTip", "Configure the behavior and features of this Editor"), NAME_None);
-		EditorSettingsContainer->DescribeCategory("General", LOCTEXT("EditorGeneralCategoryName", "General"), LOCTEXT("EditorGeneralCategoryDescription", "General Editor settings description text here"), "SettingsEditor.Category_Editor");
-		EditorSettingsContainer->DescribeCategory("LevelEditor", LOCTEXT("EditorLevelEditorCategoryName", "Level Editor"), LOCTEXT("EditorLevelEditorCategoryDescription", "Level Editor settings description text here"), "SettingsEditor.Category_LevelEditor");
-		EditorSettingsContainer->DescribeCategory("ContentEditors", LOCTEXT("EditorContentEditorsCategoryName", "Content Editors"), LOCTEXT("EditorContentEditorsCategoryDescription", "Content editors settings description text here"), "SettingsEditor.Category_ContentEditors");
-		EditorSettingsContainer->DescribeCategory("Plugins", LOCTEXT("EditorPluginsCategoryName", "Plugins"), LOCTEXT("EditorPluginsCategoryDescription", "Plugins settings description text here"), "SettingsEditor.Category_Plugins");
+		EditorSettingsContainer->DescribeCategory("General", LOCTEXT("EditorGeneralCategoryName", "General"), LOCTEXT("EditorGeneralCategoryDescription", "General Editor settings description text here"));
+		EditorSettingsContainer->DescribeCategory("LevelEditor", LOCTEXT("EditorLevelEditorCategoryName", "Level Editor"), LOCTEXT("EditorLevelEditorCategoryDescription", "Level Editor settings description text here"));
+		EditorSettingsContainer->DescribeCategory("ContentEditors", LOCTEXT("EditorContentEditorsCategoryName", "Content Editors"), LOCTEXT("EditorContentEditorsCategoryDescription", "Content editors settings description text here"));
+		EditorSettingsContainer->DescribeCategory("Plugins", LOCTEXT("EditorPluginsCategoryName", "Plugins"), LOCTEXT("EditorPluginsCategoryDescription", "Plugins settings description text here"));
 
 		// @todo gmp: move this into the modules that own these setting categories
 		TSharedRef<FSettingsContainer> ProjectSettingsContainer = FindOrAddContainer("Project");
 		ProjectSettingsContainer->Describe(LOCTEXT("ProjectSettingsSubMenuLabel", "Project Settings"), LOCTEXT("ProjectSettingsSubMenuToolTip", "Change the settings of the currently loaded project"), NAME_None);
-		ProjectSettingsContainer->DescribeCategory("Project", LOCTEXT("ProjectProjectCategoryName", "Project"), LOCTEXT("ProjectProjectCategoryDescription", "Project settings description text here"), "SettingsEditor.Category_Project");
-		ProjectSettingsContainer->DescribeCategory("Engine", LOCTEXT("ProjectEngineCategoryName", "Engine"), LOCTEXT("ProjectEngineCategoryDescription", "Project settings description text here"), "SettingsEditor.Category_Engine");
-		ProjectSettingsContainer->DescribeCategory("Platforms", LOCTEXT("ProjectPlatformsCategoryName", "Platforms"), LOCTEXT("ProjectPlatformsCategoryDescription", "Platform settings description text here"), "SettingsEditor.Category_Platforms");
-		ProjectSettingsContainer->DescribeCategory("Plugins", LOCTEXT("ProjectPluginsCategoryName", "Plugins"), LOCTEXT("ProjectPluginsCategoryDescription", "Plugins settings description text here"), "SettingsEditor.Category_Plugins");
+		ProjectSettingsContainer->DescribeCategory("Project", LOCTEXT("ProjectProjectCategoryName", "Project"), LOCTEXT("ProjectProjectCategoryDescription", "Project settings description text here"));
+		ProjectSettingsContainer->DescribeCategory("Engine", LOCTEXT("ProjectEngineCategoryName", "Engine"), LOCTEXT("ProjectEngineCategoryDescription", "Project settings description text here"));
+		ProjectSettingsContainer->DescribeCategory("Platforms", LOCTEXT("ProjectPlatformsCategoryName", "Platforms"), LOCTEXT("ProjectPlatformsCategoryDescription", "Platform settings description text here"));
+		ProjectSettingsContainer->DescribeCategory("Plugins", LOCTEXT("ProjectPluginsCategoryName", "Plugins"), LOCTEXT("ProjectPluginsCategoryDescription", "Plugins settings description text here"));
 	}
 
 	virtual void ShutdownModule( ) override { }
diff --git a/Engine/Source/Developer/Settings/Public/Interfaces/ISettingsCategory.h b/Engine/Source/Developer/Settings/Public/Interfaces/ISettingsCategory.h
--- a/Engine/Source/Developer/Settings/Public/Interfaces/ISettingsCategory.h
+++ b/Engine/Source/Developer/Settings/Public/Interfaces/ISettingsCategory.h
@@ -51,13 +51,6 @@ public:
 	 */
 	virtual int32 GetSections( TArray<ISettingsSectionPtr>& OutSections ) const = 0;
 
-	/**
-	 * Gets the name of the category's icon.
-	 *
-	 * @return Icon image name.
-	 */
-	virtual const FName& GetIconName( ) const = 0;
-
 	/**
 	 * Gets the category's name.
 	 *
diff --git a/Engine/Source/Developer/Settings/Public/Interfaces/ISettingsContainer.h b/Engine/Source/Developer/Settings/Public/Interfaces/ISettingsContainer.h
--- a/Engine/Source/Developer/Settings/Public/Interfaces/ISettingsContainer.h
+++ b/Engine/Source/Developer/Settings/Public/Interfaces/ISettingsContainer.h
@@ -51,9 +51,8 @@ public:
 	 * @param CategoryName The name of the category to update.
 	 * @param DisplayName The category's localized display name.
 	 * @param Description The category's localized description text.
-	 * @param IconName The name of the category's icon.
 	 */
-	virtual void DescribeCategory( const FName& CategoryName, const FText& DisplayName, const FText& Description, const FName& IconName ) = 0;
+	virtual void DescribeCategory( const FName& CategoryName, const FText& DisplayName, const FText& Description ) = 0;
 
 	/**
 	 * Gets the setting categories.
diff --git a/Engine/Source/Developer/SettingsEditor/Private/Widgets/SSettingsEditor.cpp b/Engine/Source/Developer/SettingsEditor/Private/Widgets/SSettingsEditor.cpp
--- a/Engine/Source/Developer/SettingsEditor/Private/Widgets/SSettingsEditor.cpp
+++ b/Engine/Source/Developer/SettingsEditor/Private/Widgets/SSettingsEditor.cpp
@@ -54,11 +54,13 @@ void SSettingsEditor::Construct( const FArguments& InArgs, const ISettingsEditor
 
 		+ SVerticalBox::Slot()
 			.FillHeight(1.0f)
+			.Padding(16.0f, 0.0f)
 			[
 				SNew(SHorizontalBox)
 
 				+ SHorizontalBox::Slot()
 					.AutoWidth()
+					.Padding(0.0f, 16.0f)
 					[
 						// categories menu
 						SNew(SScrollBox)
@@ -92,13 +94,13 @@ void SSettingsEditor::Construct( const FArguments& InArgs, const ISettingsEditor
 
 				+ SHorizontalBox::Slot()
 					.FillWidth(1.0f)
+					.Padding(0.0f, 16.0f)
 					[
 						SNew(SVerticalBox)
 							.Visibility(this, &SSettingsEditor::HandleSettingsBoxVisibility)
 
 						+ SVerticalBox::Slot()
 							.AutoHeight()
-							.Padding(0.0f, 16.0f)
 							[
 								// title and button bar
 								SNew(SHorizontalBox)
@@ -187,7 +189,7 @@ void SSettingsEditor::Construct( const FArguments& InArgs, const ISettingsEditor
 
 						+ SVerticalBox::Slot()
 							.AutoHeight()
-							.Padding(0.0f, 0.0f, 0.0f, 16.0f)
+							.Padding(0.0f, 16.0f, 0.0f, 16.0f)
 							[
 								// checkout notice
 								SNew(SSettingsEditorCheckoutNotice)
@@ -199,6 +201,7 @@ void SSettingsEditor::Construct( const FArguments& InArgs, const ISettingsEditor
 
 						+ SVerticalBox::Slot()
 							.FillHeight(1.0f)
+							.Padding(0.0f, 12.0f, 0.0f, 0.0f)
 							[
 								// settings area
 								SNew(SOverlay)
@@ -376,8 +379,6 @@ TWeakObjectPtr<UObject> SSettingsEditor::GetSelectedSettingsObject( ) const
 
 TSharedRef<SWidget> SSettingsEditor::MakeCategoryWidget( const ISettingsCategoryRef& Category )
 {
-	const float VerticalSlotPadding = 10.0f;
-
 	// create section widgets
 	TSharedRef<SVerticalBox> SectionsBox = SNew(SVerticalBox);
 	TArray<ISettingsSectionPtr> Sections;
@@ -413,7 +414,7 @@ TSharedRef<SWidget> SSettingsEditor::MakeCategoryWidget( const ISettingsCategory
 	{
 		SectionsBox->AddSlot()
 			.HAlign(HAlign_Left)
-			.Padding(0.0f, VerticalSlotPadding, 0.0f, 0.0f)
+			.Padding(0.0f, 10.0f, 0.0f, 0.0f)
 			[
 				SNew(SHorizontalBox)
 
@@ -444,44 +445,23 @@ TSharedRef<SWidget> SSettingsEditor::MakeCategoryWidget( const ISettingsCategory
 		}
 	}
 
-	const FSlateBrush* CategoryIcon = FEditorStyle::GetBrush(Category->GetIconName());
-
 	// create category widget
-	return SNew(SHorizontalBox)
+	return SNew(SVerticalBox)
 
-	+ SHorizontalBox::Slot()
-		.AutoWidth()
-		.HAlign(HAlign_Center)
-		.VAlign(VAlign_Top)
-		.Padding(8.0f, VerticalSlotPadding, 12.0f, VerticalSlotPadding)
+	+ SVerticalBox::Slot()
+		.AutoHeight()
 		[
-			// category icon
-			SNew(SImage)
-				.Image((CategoryIcon == FEditorStyle::GetDefaultBrush()) ? FEditorStyle::GetBrush("SettingsEditor.Category_Default") : CategoryIcon)
+			// category title
+			SNew(STextBlock)
+				.Font(FSlateFontInfo(FPaths::EngineContentDir() / TEXT("Slate/Fonts/Roboto-Regular.ttf"), 18))
+				.Text(Category->GetDisplayName())
 		]
 
-	+ SHorizontalBox::Slot()
-		.FillWidth(1.0f)
-		.VAlign(VAlign_Center)
+	+ SVerticalBox::Slot()
+		.FillHeight(1.0f)
 		[
-			SNew(SVerticalBox)
-
-			+ SVerticalBox::Slot()
-				.AutoHeight()
-				.Padding(0.0f, VerticalSlotPadding, 0.0f, 0.0f)
-				[
-					// category title
-					SNew(STextBlock)
-						.Font(FSlateFontInfo(FPaths::EngineContentDir() / TEXT("Slate/Fonts/Roboto-Regular.ttf"), 18))
-						.Text(Category->GetDisplayName())
-				]
-
-			+ SVerticalBox::Slot()
-				.FillHeight(1.0f)
-				[
-					// sections list
-					SectionsBox
-				]
+			// sections list
+			SectionsBox
 		];
 }
 
@@ -512,7 +492,7 @@ void SSettingsEditor::ReloadCategories( )
 	{
 		CategoriesBox->AddSlot()
 			.AutoHeight()
-			.Padding(0.0f, 0.0f, 0.0f, 8.0f)
+			.Padding(0.0f, 0.0f, 0.0f, 16.0f)
 			[
 				MakeCategoryWidget(Category.ToSharedRef())
 			];
diff --git a/Engine/Source/Editor/EditorStyle/Private/SlateEditorStyle.cpp b/Engine/Source/Editor/EditorStyle/Private/SlateEditorStyle.cpp
--- a/Engine/Source/Editor/EditorStyle/Private/SlateEditorStyle.cpp
+++ b/Engine/Source/Editor/EditorStyle/Private/SlateEditorStyle.cpp
@@ -5597,15 +5597,6 @@ void FSlateEditorStyle::FStyle::SetupAutomationStyles()
 
 	// Settings Editor
 	{
-		Set( "SettingsEditor.Category_ContentEditors", new IMAGE_BRUSH( "Icons/icon_Editor_Preferences_ContentEditors_40x", Icon40x40 ) );
-		Set( "SettingsEditor.Category_Default", new IMAGE_BRUSH( "Icons/icon_MatEd_Sphere_40x", Icon40x40 ) );
-		Set( "SettingsEditor.Category_Editor", new IMAGE_BRUSH( "Icons/icon_Editor_Preferences_General_40x", Icon40x40 ) );
-		Set( "SettingsEditor.Category_Engine", new IMAGE_BRUSH( "Icons/icon_Cascade_RestartSim_40x", Icon40x40 ) );
-		Set( "SettingsEditor.Category_Game", new IMAGE_BRUSH( "Icons/icon_MatEd_Realtime_40x", Icon40x40 ) );
-		Set( "SettingsEditor.Category_LevelEditor", new IMAGE_BRUSH( "Icons/icon_Editor_Preferences_LevelEditor_40x", Icon40x40 ) );
-		Set( "SettingsEditor.Category_Project", new IMAGE_BRUSH( "Icons/icon_BlueprintEditor_Settings_40x", Icon40x40 ) );
-		Set( "SettingsEditor.Category_Platforms", new IMAGE_BRUSH( "Icons/icon_PlayOnDevice_40px", Icon40x40 ) );
-		Set( "SettingsEditor.Category_Plugins", new IMAGE_BRUSH( "Icons/Profiler/profiler_mem_40x", Icon40x40 ) );
 		Set( "SettingsEditor.Collision_Engine", new IMAGE_BRUSH("Icons/icon_Cascade_RestartSim_40x", Icon16x16));
 		Set( "SettingsEditor.Collision_Game", new IMAGE_BRUSH("Icons/icon_MatEd_Realtime_40x", Icon16x16));
 
