Settings: Various improvements to the API

- made public headers compilable individually
- easier access to settings section delegates
- removed module singleton accessor
- moved non-trivial definitions into cpp files
- code & documentation cleanup

#UpgradeNotes:

- instead of ISettingsModule::Get() use FModuleManager::GetModulePtr<ISettingsModule>("Settings")
- instead of using FSettingsSectionDelegates assign delegates directly through the new ISettingsSection methods

diff --git a/Engine/Plugins/2D/Paper2D/Source/Paper2DEditor/Private/Paper2DEditorModule.cpp b/Engine/Plugins/2D/Paper2D/Source/Paper2DEditor/Private/Paper2DEditorModule.cpp
--- a/Engine/Plugins/2D/Paper2D/Source/Paper2DEditor/Private/Paper2DEditorModule.cpp
+++ b/Engine/Plugins/2D/Paper2D/Source/Paper2DEditor/Private/Paper2DEditorModule.cpp
@@ -9,7 +9,7 @@
 #include "PaperEditorCommands.h"
 
 #include "AssetEditorToolkit.h"
-
+#include "ModuleManager.h"
 #include "ContentBrowserExtensions/ContentBrowserExtensions.h"
 
 // Sprite support
@@ -37,7 +37,7 @@
 
 // Settings
 #include "PaperRuntimeSettings.h"
-#include "Settings.h"
+#include "ISettingsModule.h"
 
 // Intro tutorials
 #include "Editor/IntroTutorials/Public/IIntroTutorials.h"
@@ -212,19 +212,19 @@ private:
 
 	void RegisterSettings()
 	{
-		if (ISettingsModule* SettingsModule = ISettingsModule::Get())
+		if (ISettingsModule* SettingsModule = FModuleManager::GetModulePtr<ISettingsModule>("Settings"))
 		{
 			SettingsModule->RegisterSettings("Project", "Plugins", "Paper2D",
 				LOCTEXT("RuntimeSettingsName", "Paper 2D"),
 				LOCTEXT("RuntimeSettingsDescription", "Configure the Paper 2D plugin"),
 				GetMutableDefault<UPaperRuntimeSettings>()
-				);
+			);
 		}
 	}
 
 	void UnregisterSettings()
 	{
-		if (ISettingsModule* SettingsModule = ISettingsModule::Get())
+		if (ISettingsModule* SettingsModule = FModuleManager::GetModulePtr<ISettingsModule>("Settings"))
 		{
 			SettingsModule->UnregisterSettings("Project", "Plugins", "Paper2D");
 		}
diff --git a/Engine/Plugins/Editor/MacGraphicsSwitching/Source/MacGraphicsSwitching/Private/MacGraphicsSwitchingModule.cpp b/Engine/Plugins/Editor/MacGraphicsSwitching/Source/MacGraphicsSwitching/Private/MacGraphicsSwitchingModule.cpp
--- a/Engine/Plugins/Editor/MacGraphicsSwitching/Source/MacGraphicsSwitching/Private/MacGraphicsSwitchingModule.cpp
+++ b/Engine/Plugins/Editor/MacGraphicsSwitching/Source/MacGraphicsSwitching/Private/MacGraphicsSwitchingModule.cpp
@@ -10,7 +10,9 @@
 #include "Runtime/SlateCore/Public/Rendering/SlateRenderer.h"
 #include "Developer/Settings/Public/Settings.h"
 #include "Editor/MainFrame/Public/Interfaces/IMainFrameModule.h"
+#include "ISettingsModule.h"
 #include "LevelEditor.h"
+#include "ModuleManager.h"
 
 #import <IOKit/ps/IOPowerSources.h>
 #import <IOKit/ps/IOPSKeys.h>
@@ -103,7 +105,7 @@ void FMacGraphicsSwitchingModule::StartupModule()
 	}
 	
 	// Register settings
-	ISettingsModule* SettingsModule = ISettingsModule::Get();
+	ISettingsModule* SettingsModule = FModuleManager::GetModulePtr<ISettingsModule>("Settings");
 	if( SettingsModule != nullptr )
 	{
 		SettingsModule->RegisterSettings( "Editor", "Plugins", "MacGraphicsSwitching",
@@ -161,7 +163,7 @@ void FMacGraphicsSwitchingModule::AddGraphicsSwitcher( FToolBarBuilder& ToolBarB
 void FMacGraphicsSwitchingModule::ShutdownModule()
 {
 	// Unregister settings
-	ISettingsModule* SettingsModule = ISettingsModule::Get();
+	ISettingsModule* SettingsModule = FModuleManager::GetModulePtr<ISettingsModule>("Settings");
 	if( SettingsModule != nullptr )
 	{
 		SettingsModule->UnregisterSettings( "Editor", "Plugins", "MacGraphicsSwitching" );
diff --git a/Engine/Plugins/Messaging/UdpMessaging/Source/UdpMessaging/Private/UdpMessagingModule.cpp b/Engine/Plugins/Messaging/UdpMessaging/Source/UdpMessaging/Private/UdpMessagingModule.cpp
--- a/Engine/Plugins/Messaging/UdpMessaging/Source/UdpMessaging/Private/UdpMessagingModule.cpp
+++ b/Engine/Plugins/Messaging/UdpMessaging/Source/UdpMessaging/Private/UdpMessagingModule.cpp
@@ -1,6 +1,8 @@
 // Copyright 1998-2014 Epic Games, Inc. All Rights Reserved.
 
 #include "UdpMessagingPrivatePCH.h"
+#include "ISettingsModule.h"
+#include "ISettingsSection.h"
 #include "ModuleInterface.h"
 #include "ModuleManager.h"
 
@@ -34,19 +36,20 @@ public:
 		}
 
 		// register settings
-		ISettingsModule* SettingsModule = ISettingsModule::Get();
+		ISettingsModule* SettingsModule = FModuleManager::GetModulePtr<ISettingsModule>("Settings");
 
 		if (SettingsModule != nullptr)
 		{
-			FSettingsSectionDelegates SettingsDelegates;
-			SettingsDelegates.ModifiedDelegate = FOnSettingsSectionModified::CreateRaw(this, &FUdpMessagingModule::HandleSettingsSaved);
-
-			SettingsModule->RegisterSettings("Project", "Plugins", "UdpMessaging",
+			ISettingsSectionPtr SettingsSection = SettingsModule->RegisterSettings("Project", "Plugins", "UdpMessaging",
 				LOCTEXT("UdpMessagingSettingsName", "UDP Messaging"),
 				LOCTEXT("UdpMessagingSettingsDescription", "Configure the UDP Messaging plug-in."),
-				GetMutableDefault<UUdpMessagingSettings>(),
-				SettingsDelegates
+				GetMutableDefault<UUdpMessagingSettings>()
 			);
+
+			if (SettingsSection.IsValid())
+			{
+				SettingsSection->OnModified().BindRaw(this, &FUdpMessagingModule::HandleSettingsSaved);
+			}
 		}
 
 		// register application events
@@ -63,7 +66,7 @@ public:
 		FCoreDelegates::ApplicationWillDeactivateDelegate.RemoveAll(this);
 
 		// unregister settings
-		ISettingsModule* SettingsModule = ISettingsModule::Get();
+		ISettingsModule* SettingsModule = FModuleManager::GetModulePtr<ISettingsModule>("Settings");
 
 		if (SettingsModule != nullptr)
 		{
diff --git a/Engine/Plugins/Messaging/UdpMessaging/Source/UdpMessaging/Private/UdpMessagingPrivatePCH.h b/Engine/Plugins/Messaging/UdpMessaging/Source/UdpMessaging/Private/UdpMessagingPrivatePCH.h
--- a/Engine/Plugins/Messaging/UdpMessaging/Source/UdpMessaging/Private/UdpMessagingPrivatePCH.h
+++ b/Engine/Plugins/Messaging/UdpMessaging/Source/UdpMessaging/Private/UdpMessagingPrivatePCH.h
@@ -10,7 +10,6 @@
 #include "CoreUObject.h"
 #include "Messaging.h"
 #include "Networking.h"
-#include "Settings.h"
 #include "Sockets.h"
 #include "SocketSubsystem.h"
 
diff --git a/Engine/Plugins/Runtime/TwitchLiveStreaming/Source/TwitchLiveStreaming/TwitchLiveStreaming.cpp b/Engine/Plugins/Runtime/TwitchLiveStreaming/Source/TwitchLiveStreaming/TwitchLiveStreaming.cpp
--- a/Engine/Plugins/Runtime/TwitchLiveStreaming/Source/TwitchLiveStreaming/TwitchLiveStreaming.cpp
+++ b/Engine/Plugins/Runtime/TwitchLiveStreaming/Source/TwitchLiveStreaming/TwitchLiveStreaming.cpp
@@ -6,13 +6,15 @@
 #include "ModuleManager.h"
 #include "Runtime/Core/Public/Features/IModularFeatures.h"
 #include "Runtime/Core/Public/Stats/Stats2.h"
-#include "Settings.h"
+#include "ISettingsModule.h"
+
 
 DEFINE_LOG_CATEGORY_STATIC( LogTwitch, Log, All );
 
+
 #if WITH_TWITCH
-IMPLEMENT_MODULE( FTwitchLiveStreaming, TwitchLiveStreaming )
 
+IMPLEMENT_MODULE( FTwitchLiveStreaming, TwitchLiveStreaming )
 DECLARE_MEMORY_STAT( TEXT( "Twitch Memory" ), STAT_TwitchMemory, STATGROUP_Memory );
 
 #define LOCTEXT_NAMESPACE "TwitchPlugin"
@@ -94,7 +96,7 @@ void FTwitchLiveStreaming::StartupModule()
 	// @todo twitch: Editor: Ideally we want a "zoom in" feature to use during editor live streams
 
 	// Register our custom project settings
-	ISettingsModule* SettingsModule = ISettingsModule::Get();
+	ISettingsModule* SettingsModule = FModuleManager::GetModulePtr<ISettingsModule>("Settings");
 	if( SettingsModule != nullptr )
 	{
 		SettingsModule->RegisterSettings(
@@ -115,7 +117,7 @@ void FTwitchLiveStreaming::ShutdownModule()
 	IModularFeatures::Get().UnregisterModularFeature( TEXT( "LiveStreaming" ), this );
 
 	// Unregister custom project settings
-	ISettingsModule* SettingsModule = ISettingsModule::Get();
+	ISettingsModule* SettingsModule = FModuleManager::GetModulePtr<ISettingsModule>("Settings");
 	if( SettingsModule != nullptr )
 	{
 		SettingsModule->UnregisterSettings( "Project", "Plugins", "Twitch" );
diff --git a/Engine/Plugins/Slate/SlateRemote/Source/SlateRemote/Private/SlateRemoteModule.cpp b/Engine/Plugins/Slate/SlateRemote/Source/SlateRemote/Private/SlateRemoteModule.cpp
--- a/Engine/Plugins/Slate/SlateRemote/Source/SlateRemote/Private/SlateRemoteModule.cpp
+++ b/Engine/Plugins/Slate/SlateRemote/Source/SlateRemote/Private/SlateRemoteModule.cpp
@@ -1,6 +1,10 @@
 // Copyright 1998-2014 Epic Games, Inc. All Rights Reserved.
 
 #include "SlateRemotePrivatePCH.h"
+#include "ISettingsModule.h"
+#include "ISettingsSection.h"
+#include "ModuleInterface.h"
+#include "ModuleManager.h"
 
 
 #define LOCTEXT_NAMESPACE "FSlateRemoteModule"
@@ -24,19 +28,20 @@ public:
 		}
 
 		// register settings
-		ISettingsModule* SettingsModule = ISettingsModule::Get();
+		ISettingsModule* SettingsModule = FModuleManager::GetModulePtr<ISettingsModule>("Settings");
 
 		if (SettingsModule != nullptr)
 		{
-			FSettingsSectionDelegates SettingsDelegates;
-			SettingsDelegates.ModifiedDelegate = FOnSettingsSectionModified::CreateRaw(this, &FSlateRemoteModule::HandleSettingsSaved);
-
-			SettingsModule->RegisterSettings("Project", "Plugins", "SlateRemote",
+			ISettingsSectionPtr SettingsSection = SettingsModule->RegisterSettings("Project", "Plugins", "SlateRemote",
 				LOCTEXT("SlateRemoteSettingsName", "Slate Remote"),
 				LOCTEXT("SlateRemoteSettingsDescription", "Configure the Slate Remote plug-in."),
-				GetMutableDefault<USlateRemoteSettings>(),
-				SettingsDelegates
+				GetMutableDefault<USlateRemoteSettings>()
 			);
+
+			if (SettingsSection.IsValid())
+			{
+				SettingsSection->OnModified().BindRaw(this, &FSlateRemoteModule::HandleSettingsSaved);
+			}
 		}
 
 		// register application events
@@ -53,7 +58,7 @@ public:
 		FCoreDelegates::ApplicationWillDeactivateDelegate.RemoveAll(this);
 
 		// unregister settings
-		ISettingsModule* SettingsModule = ISettingsModule::Get();
+		ISettingsModule* SettingsModule = FModuleManager::GetModulePtr<ISettingsModule>("Settings");
 
 		if (SettingsModule != nullptr)
 		{
diff --git a/Engine/Plugins/Slate/SlateRemote/Source/SlateRemote/Private/SlateRemotePrivatePCH.h b/Engine/Plugins/Slate/SlateRemote/Source/SlateRemote/Private/SlateRemotePrivatePCH.h
--- a/Engine/Plugins/Slate/SlateRemote/Source/SlateRemote/Private/SlateRemotePrivatePCH.h
+++ b/Engine/Plugins/Slate/SlateRemote/Source/SlateRemote/Private/SlateRemotePrivatePCH.h
@@ -7,9 +7,7 @@
  *****************************************************************************/
 
 #include "Core.h"
-#include "ModuleManager.h"
 #include "Networking.h"
-#include "Settings.h"
 #include "SlateBasics.h"
 #include "SlateCore.h"
 #include "Sockets.h"
diff --git a/Engine/Source/Developer/Android/AndroidPlatformEditor/Private/AndroidPlatformEditorModule.cpp b/Engine/Source/Developer/Android/AndroidPlatformEditor/Private/AndroidPlatformEditorModule.cpp
--- a/Engine/Source/Developer/Android/AndroidPlatformEditor/Private/AndroidPlatformEditorModule.cpp
+++ b/Engine/Source/Developer/Android/AndroidPlatformEditor/Private/AndroidPlatformEditorModule.cpp
@@ -2,26 +2,35 @@
 
 #include "AndroidPlatformEditorPrivatePCH.h"
 #include "AndroidTargetSettingsCustomization.h"
-#include "Settings.h"
+#include "ModuleInterface.h"
+#include "ISettingsModule.h"
+#include "ModuleManager.h"
+
 
 #define LOCTEXT_NAMESPACE "FAndroidPlatformEditorModule"
 
+
 /**
  * Module for Android platform editor utilities
  */
-class FAndroidPlatformEditorModule : public IModuleInterface
+class FAndroidPlatformEditorModule
+	: public IModuleInterface
 {
 	// IModuleInterface interface
+
 	virtual void StartupModule() override
 	{
-		// Register the settings detail panel customization
+		// register settings detail panel customization
 		FPropertyEditorModule& PropertyModule = FModuleManager::LoadModuleChecked<FPropertyEditorModule>("PropertyEditor");
 		PropertyModule.RegisterCustomClassLayout(
 			UAndroidRuntimeSettings::StaticClass()->GetFName(),
-			FOnGetDetailCustomizationInstance::CreateStatic(&FAndroidTargetSettingsCustomization::MakeInstance));
+			FOnGetDetailCustomizationInstance::CreateStatic(&FAndroidTargetSettingsCustomization::MakeInstance)
+		);
+
 		PropertyModule.NotifyCustomizationModuleChanged();
 
-		ISettingsModule* SettingsModule = ISettingsModule::Get();
+		// register settings
+		ISettingsModule* SettingsModule = FModuleManager::GetModulePtr<ISettingsModule>("Settings");
 
 		if (SettingsModule != nullptr)
 		{
@@ -35,16 +44,16 @@ class FAndroidPlatformEditorModule : public IModuleInterface
 
 	virtual void ShutdownModule() override
 	{
-		ISettingsModule* SettingsModule = ISettingsModule::Get();
+		ISettingsModule* SettingsModule = FModuleManager::GetModulePtr<ISettingsModule>("Settings");
 
 		if (SettingsModule != nullptr)
 		{
 			SettingsModule->UnregisterSettings("Project", "Platforms", "Android");
 		}
 	}
-	// End of IModuleInterface interface
 };
 
+
 IMPLEMENT_MODULE(FAndroidPlatformEditorModule, AndroidPlatformEditor);
 
 #undef LOCTEXT_NAMESPACE
diff --git a/Engine/Source/Developer/Android/AndroidPlatformEditor/Private/AndroidPlatformEditorPrivatePCH.h b/Engine/Source/Developer/Android/AndroidPlatformEditor/Private/AndroidPlatformEditorPrivatePCH.h
--- a/Engine/Source/Developer/Android/AndroidPlatformEditor/Private/AndroidPlatformEditorPrivatePCH.h
+++ b/Engine/Source/Developer/Android/AndroidPlatformEditor/Private/AndroidPlatformEditorPrivatePCH.h
@@ -2,7 +2,9 @@
 
 #pragma once
 
+
+/* Private dependencies
+ *****************************************************************************/
+
 #include "Core.h"
-#include "ModuleInterface.h"
-#include "ModuleManager.h"
 #include "Engine.h"
diff --git a/Engine/Source/Developer/HTML5/HTML5TargetPlatform/Private/HTML5TargetPlatformModule.cpp b/Engine/Source/Developer/HTML5/HTML5TargetPlatform/Private/HTML5TargetPlatformModule.cpp
--- a/Engine/Source/Developer/HTML5/HTML5TargetPlatform/Private/HTML5TargetPlatformModule.cpp
+++ b/Engine/Source/Developer/HTML5/HTML5TargetPlatform/Private/HTML5TargetPlatformModule.cpp
@@ -1,10 +1,8 @@
 // Copyright 1998-2014 Epic Games, Inc. All Rights Reserved.
 
-/*=============================================================================
-	HTML5TargetPlatformModule.cpp: Implements the FHTML5TargetPlatformModule class.
-=============================================================================*/
-
 #include "HTML5TargetPlatformPrivatePCH.h"
+#include "ISettingsModule.h"
+#include "ModuleManager.h"
 
 
 #define LOCTEXT_NAMESPACE "FHTML5TargetPlatformModule"
@@ -63,7 +61,7 @@ public:
 		TargetSettings = ConstructObject<UHTML5TargetSettings>(UHTML5TargetSettings::StaticClass(), GetTransientPackage(), "HTML5TargetSettings", RF_Standalone);
 		TargetSettings->AddToRoot();
 
-		ISettingsModule* SettingsModule = ISettingsModule::Get();
+		ISettingsModule* SettingsModule = FModuleManager::GetModulePtr<ISettingsModule>("Settings");
 
 		if (SettingsModule != nullptr)
 		{
@@ -77,7 +75,7 @@ public:
 
 	virtual void ShutdownModule() override
 	{
-		ISettingsModule* SettingsModule = ISettingsModule::Get();
+		ISettingsModule* SettingsModule = FModuleManager::GetModulePtr<ISettingsModule>("Settings");
 
 		if (SettingsModule != nullptr)
 		{
diff --git a/Engine/Source/Developer/HTML5/HTML5TargetPlatform/Private/HTML5TargetPlatformPrivatePCH.h b/Engine/Source/Developer/HTML5/HTML5TargetPlatform/Private/HTML5TargetPlatformPrivatePCH.h
--- a/Engine/Source/Developer/HTML5/HTML5TargetPlatform/Private/HTML5TargetPlatformPrivatePCH.h
+++ b/Engine/Source/Developer/HTML5/HTML5TargetPlatform/Private/HTML5TargetPlatformPrivatePCH.h
@@ -1,17 +1,12 @@
 // Copyright 1998-2014 Epic Games, Inc. All Rights Reserved.
 
-/*=============================================================================
-	HTML5TargetPlatformPrivatePCH.h: Pre-compiled header file for the HTML5TargetPlatform module.
-=============================================================================*/
-
 #pragma once
 
+
 /* Dependencies
  *****************************************************************************/
 
 #include "Core.h"
-#include "ModuleManager.h"
-#include "Settings.h"
 #include "Runtime/Core/Public/HTML5/HTML5PlatformProperties.h"
 
 #if WITH_ENGINE
diff --git a/Engine/Source/Developer/Linux/LinuxTargetPlatform/Private/LinuxTargetPlatformModule.cpp b/Engine/Source/Developer/Linux/LinuxTargetPlatform/Private/LinuxTargetPlatformModule.cpp
--- a/Engine/Source/Developer/Linux/LinuxTargetPlatform/Private/LinuxTargetPlatformModule.cpp
+++ b/Engine/Source/Developer/Linux/LinuxTargetPlatform/Private/LinuxTargetPlatformModule.cpp
@@ -1,10 +1,9 @@
 // Copyright 1998-2014 Epic Games, Inc. All Rights Reserved.
 
-/*=============================================================================
-	LinuxTargetPlatformModule.cpp: Implements the FAndroidTargetPlatformModule class.
-=============================================================================*/
-
 #include "LinuxTargetPlatformPrivatePCH.h"
+#include "ISettingsModule.h"
+#include "ModuleInterface.h"
+#include "ModuleManager.h"
 
 
 #define LOCTEXT_NAMESPACE "FLinuxTargetPlatformModule"
@@ -24,18 +23,15 @@ class FLinuxTargetPlatformModule
 {
 public:
 
-	/**
-	 * Destructor.
-	 */
+	/** Destructor. */
 	~FLinuxTargetPlatformModule( )
 	{
 		Singleton = NULL;
 	}
 
-
 public:
 	
-	// Begin ITargetPlatformModule interface
+	// ITargetPlatformModule interface
 
 	virtual ITargetPlatform* GetTargetPlatform( ) override
 	{
@@ -47,12 +43,9 @@ public:
 		return Singleton;
 	}
 
-	// End ITargetPlatformModule interface
-
-
 public:
 
-	// Begin IModuleInterface interface
+	// IModuleInterface interface
 
 	virtual void StartupModule() override
 	{
@@ -62,7 +55,7 @@ public:
 		GConfig->GetArray(TEXT("/Script/LinuxTargetPlatform.LinuxTargetSettings"), TEXT("TargetedRHIs"), TargetSettings->TargetedRHIs, GEngineIni);
 		TargetSettings->AddToRoot();
 
-		ISettingsModule* SettingsModule = ISettingsModule::Get();
+		ISettingsModule* SettingsModule = FModuleManager::GetModulePtr<ISettingsModule>("Settings");
 
 		if (SettingsModule != nullptr)
 		{
@@ -76,7 +69,7 @@ public:
 
 	virtual void ShutdownModule() override
 	{
-		ISettingsModule* SettingsModule = ISettingsModule::Get();
+		ISettingsModule* SettingsModule = FModuleManager::GetModulePtr<ISettingsModule>("Settings");
 
 		if (SettingsModule != nullptr)
 		{
@@ -93,12 +86,10 @@ public:
 			TargetSettings = NULL;
 		}
 	}
-	// End IModuleInterface interface
-
 
 private:
 
-	// Holds the target settings.
+	/** Holds the target settings. */
 	ULinuxTargetSettings* TargetSettings;
 };
 
diff --git a/Engine/Source/Developer/Linux/LinuxTargetPlatform/Private/LinuxTargetPlatformPrivatePCH.h b/Engine/Source/Developer/Linux/LinuxTargetPlatform/Private/LinuxTargetPlatformPrivatePCH.h
--- a/Engine/Source/Developer/Linux/LinuxTargetPlatform/Private/LinuxTargetPlatformPrivatePCH.h
+++ b/Engine/Source/Developer/Linux/LinuxTargetPlatform/Private/LinuxTargetPlatformPrivatePCH.h
@@ -1,9 +1,5 @@
 // Copyright 1998-2014 Epic Games, Inc. All Rights Reserved.
 
-/*=============================================================================
-	LinuxTargetPlatformPrivatePCH.h: Pre-compiled header file for the LinuxTargetPlatform module.
-=============================================================================*/
-
 #pragma once
 
 
@@ -11,9 +7,6 @@
  *****************************************************************************/
 
 #include "Core.h"
-#include "ModuleInterface.h"
-#include "ModuleManager.h"
-#include "Settings.h"
 #include "Runtime/Core/Public/Linux/LinuxPlatformProperties.h"
 
 #if WITH_ENGINE
diff --git a/Engine/Source/Developer/Mac/MacTargetPlatform/Private/MacTargetPlatformModule.cpp b/Engine/Source/Developer/Mac/MacTargetPlatform/Private/MacTargetPlatformModule.cpp
--- a/Engine/Source/Developer/Mac/MacTargetPlatform/Private/MacTargetPlatformModule.cpp
+++ b/Engine/Source/Developer/Mac/MacTargetPlatform/Private/MacTargetPlatformModule.cpp
@@ -1,10 +1,8 @@
 // Copyright 1998-2014 Epic Games, Inc. All Rights Reserved.
 
-/*=============================================================================
-	MacTargetPlatformModule.cpp: Implements the FMacTargetPlatformModule class.
-=============================================================================*/
-
 #include "MacTargetPlatformPrivatePCH.h"
+#include "ModuleManager.h"
+#include "ISettingsModule.h"
 
 
 #define LOCTEXT_NAMESPACE "FMacTargetPlatformModule"
@@ -59,7 +57,7 @@ public:
 		TargetSettings = ConstructObject<UMacTargetSettings>(UMacTargetSettings::StaticClass(), GetTransientPackage(), "MacTargetSettings", RF_Standalone);
 		TargetSettings->AddToRoot();
 
-		ISettingsModule* SettingsModule = ISettingsModule::Get();
+		ISettingsModule* SettingsModule = FModuleManager::GetModulePtr<ISettingsModule>("Settings");
 
 		if (SettingsModule != nullptr)
 		{
@@ -73,7 +71,7 @@ public:
 
 	virtual void ShutdownModule() override
 	{
-		ISettingsModule* SettingsModule = ISettingsModule::Get();
+		ISettingsModule* SettingsModule = FModuleManager::GetModulePtr<ISettingsModule>("Settings");
 
 		if (SettingsModule != nullptr)
 		{
diff --git a/Engine/Source/Developer/Mac/MacTargetPlatform/Private/MacTargetPlatformPrivatePCH.h b/Engine/Source/Developer/Mac/MacTargetPlatform/Private/MacTargetPlatformPrivatePCH.h
--- a/Engine/Source/Developer/Mac/MacTargetPlatform/Private/MacTargetPlatformPrivatePCH.h
+++ b/Engine/Source/Developer/Mac/MacTargetPlatform/Private/MacTargetPlatformPrivatePCH.h
@@ -1,9 +1,5 @@
 // Copyright 1998-2014 Epic Games, Inc. All Rights Reserved.
 
-/*=============================================================================
-	MacTargetPlatformPrivatePCH.h: Pre-compiled header file for the MacTargetPlatform module.
-=============================================================================*/
-
 #pragma once
 
 
@@ -11,9 +7,6 @@
  *****************************************************************************/
 
 #include "Core.h"
-#include "ModuleInterface.h"
-#include "ModuleManager.h"
-#include "Settings.h"
 #include "Runtime/Core/Public/Mac/MacPlatformProperties.h"
 
 #if WITH_ENGINE
@@ -24,6 +17,7 @@
 #include "TargetPlatform.h"
 #include "TargetPlatformBase.h"
 
+
 /* Private includes
  *****************************************************************************/
 
diff --git a/Engine/Source/Developer/Settings/Private/SettingsCategory.cpp b/Engine/Source/Developer/Settings/Private/SettingsCategory.cpp
new file mode 100644
--- /dev/null
+++ b/Engine/Source/Developer/Settings/Private/SettingsCategory.cpp
@@ -0,0 +1,69 @@
+// Copyright 1998-2014 Epic Games, Inc. All Rights Reserved.
+
+#include "SettingsPrivatePCH.h"
+
+
+/* FSettingsCategory structors
+ *****************************************************************************/
+
+FSettingsCategory::FSettingsCategory( const FName& InName )
+	: Name(InName)
+{ }
+
+
+/* FSettingsCategory interface
+ *****************************************************************************/
+
+ISettingsSectionRef FSettingsCategory::AddSection( const FName& SectionName, const FText& InDisplayName, const FText& InDescription, const TWeakObjectPtr<UObject>& SettingsObject )
+{
+	TSharedPtr<FSettingsSection>& Section = Sections.FindOrAdd(SectionName);
+
+	if (!Section.IsValid() || (Section->GetSettingsObject() != SettingsObject) || Section->GetCustomWidget().IsValid())
+	{
+		Section = MakeShareable(new FSettingsSection(AsShared(), SectionName, InDisplayName, InDescription, SettingsObject));
+	}
+
+	return Section.ToSharedRef();
+}
+
+
+ISettingsSectionRef FSettingsCategory::AddSection( const FName& SectionName, const FText& InDisplayName, const FText& InDescription, const TSharedRef<SWidget>& CustomWidget )
+{
+	TSharedPtr<FSettingsSection>& Section = Sections.FindOrAdd(SectionName);
+
+	if (!Section.IsValid() || (Section->GetSettingsObject() != nullptr) || (Section->GetCustomWidget().Pin() != CustomWidget))
+	{
+		Section = MakeShareable(new FSettingsSection(AsShared(), SectionName, InDisplayName, InDescription, CustomWidget));
+	}
+
+	return Section.ToSharedRef();
+}
+
+
+void FSettingsCategory::Describe( const FText& InDisplayName, const FText& InDescription )
+{
+	Description = InDescription;
+	DisplayName = InDisplayName;
+}
+
+
+void FSettingsCategory::RemoveSection( const FName& SectionName )
+{
+	Sections.Remove(SectionName);
+}
+
+
+/* ISettingsCategory interface
+ *****************************************************************************/
+
+int32 FSettingsCategory::GetSections( TArray<ISettingsSectionPtr>& OutSections ) const
+{
+	OutSections.Empty(Sections.Num());
+
+	for (TMap<FName, TSharedPtr<FSettingsSection> >::TConstIterator It(Sections); It; ++It)
+	{
+		OutSections.Add(It.Value());
+	}
+
+	return OutSections.Num();
+}
diff --git a/Engine/Source/Developer/Settings/Private/SettingsCategory.h b/Engine/Source/Developer/Settings/Private/SettingsCategory.h
--- a/Engine/Source/Developer/Settings/Private/SettingsCategory.h
+++ b/Engine/Source/Developer/Settings/Private/SettingsCategory.h
@@ -17,9 +17,7 @@ public:
 	 *
 	 * @param InName The category's name.
 	 */
-	FSettingsCategory( const FName& InName )
-		: Name(InName)
-	{ }
+	FSettingsCategory( const FName& InName );
 
 public:
 
@@ -32,20 +30,9 @@ public:
 	 * @param DisplayName The section's localized display name.
 	 * @param Description The section's localized description text.
 	 * @param SettingsObject The object that holds the section's settings.
-	 * @param Delegates The section's optional callback delegates.
 	 * @return The added settings section.
 	 */
-	ISettingsSectionRef AddSection( const FName& SectionName, const FText& InDisplayName, const FText& InDescription, const TWeakObjectPtr<UObject>& SettingsObject, const FSettingsSectionDelegates& Delegates )
-	{
-		TSharedPtr<FSettingsSection>& Section = Sections.FindOrAdd(SectionName);
-
-		if (!Section.IsValid() || (Section->GetSettingsObject() != SettingsObject) || Section->GetCustomWidget().IsValid())
-		{
-			Section = MakeShareable(new FSettingsSection(AsShared(), SectionName, InDisplayName, InDescription, Delegates, SettingsObject));
-		}
-
-		return Section.ToSharedRef();
-	}
+	ISettingsSectionRef AddSection( const FName& SectionName, const FText& InDisplayName, const FText& InDescription, const TWeakObjectPtr<UObject>& SettingsObject );
 
 	/**
 	 * Adds a settings section to this settings category.
@@ -56,20 +43,9 @@ public:
 	 * @param DisplayName The section's localized display name.
 	 * @param Description The section's localized description text.
 	 * @param CustomWidget A custom settings widget.
-	 * @param Delegates The section's optional callback delegates.
 	 * @return The added settings section.
 	 */
-	ISettingsSectionRef AddSection( const FName& SectionName, const FText& InDisplayName, const FText& InDescription, const TSharedRef<SWidget>& CustomWidget, const FSettingsSectionDelegates& Delegates )
-	{
-		TSharedPtr<FSettingsSection>& Section = Sections.FindOrAdd(SectionName);
-
-		if (!Section.IsValid() || (Section->GetSettingsObject() != nullptr) || (Section->GetCustomWidget().Pin() != CustomWidget))
-		{
-			Section = MakeShareable(new FSettingsSection(AsShared(), SectionName, InDisplayName, InDescription, Delegates, CustomWidget));
-		}
-
-		return Section.ToSharedRef();
-	}
+	ISettingsSectionRef AddSection( const FName& SectionName, const FText& InDisplayName, const FText& InDescription, const TSharedRef<SWidget>& CustomWidget );
 
 	/**
 	 * Updates the details of this settings category.
@@ -79,37 +55,30 @@ public:
 	 * @param InIconName The name of the category's icon.
 	 * @return The category.
 	 */
-	void Describe( const FText& InDisplayName, const FText& InDescription )
-	{
-		Description = InDescription;
-		DisplayName = InDisplayName;
-	}
+	void Describe( const FText& InDisplayName, const FText& InDescription );
 
 	/**
 	 * Removes a settings section.
 	 *
 	 * @param SectionName The name of the section to remove.
 	 */
-	void RemoveSection( const FName& SectionName )
-	{
-		Sections.Remove(SectionName);
-	}
+	void RemoveSection( const FName& SectionName );
 
 public:
 
 	// ISettingsCategory interface
 
-	virtual const FText& GetDescription( ) const override
+	virtual const FText& GetDescription() const override
 	{
 		return Description;
 	}
 
-	virtual const FText& GetDisplayName( ) const override
+	virtual const FText& GetDisplayName() const override
 	{
 		return DisplayName;
 	}
 
-	virtual const FName& GetName( ) const override
+	virtual const FName& GetName() const override
 	{
 		return Name;
 	}
@@ -119,29 +88,19 @@ public:
 		return Sections.FindRef(SectionName);
 	}
 
-	virtual int32 GetSections( TArray<ISettingsSectionPtr>& OutSections ) const override
-	{
-		OutSections.Empty(Sections.Num());
-
-		for (TMap<FName, TSharedPtr<FSettingsSection> >::TConstIterator It(Sections); It; ++It)
-		{
-			OutSections.Add(It.Value());
-		}
-
-		return OutSections.Num();
-	}
+	virtual int32 GetSections( TArray<ISettingsSectionPtr>& OutSections ) const override;
 
 private:
 
-	// Holds the category's description text.
+	/** Holds the category's description text. */
 	FText Description;
 
-	// Holds the category's localized display name.
+	/** Holds the category's localized display name. */
 	FText DisplayName;
 
-	// Holds the collection of setting sections.
-	TMap<FName, TSharedPtr<FSettingsSection> > Sections;
+	/** Holds the collection of setting sections. */
+	TMap<FName, TSharedPtr<FSettingsSection>> Sections;
 
-	// Holds the category's name.
+	/** Holds the category's name. */
 	FName Name;
 };
diff --git a/Engine/Source/Developer/Settings/Private/SettingsContainer.cpp b/Engine/Source/Developer/Settings/Private/SettingsContainer.cpp
new file mode 100644
--- /dev/null
+++ b/Engine/Source/Developer/Settings/Private/SettingsContainer.cpp
@@ -0,0 +1,101 @@
+// Copyright 1998-2014 Epic Games, Inc. All Rights Reserved.
+
+#include "SettingsPrivatePCH.h"
+
+
+/* FSettingsContainer structors
+ *****************************************************************************/
+
+FSettingsContainer::FSettingsContainer( const FName& InName )
+	: Name(InName)
+{ }
+
+
+/* FSettingsContainer interface
+ *****************************************************************************/
+
+ISettingsSectionPtr FSettingsContainer::AddSection( const FName& CategoryName, const FName& SectionName, const FText& InDisplayName, const FText& InDescription, const TWeakObjectPtr<UObject>& SettingsObject )
+{
+	TSharedPtr<FSettingsCategory> Category = Categories.FindRef(CategoryName);
+
+	if (!Category.IsValid())
+	{
+		return nullptr;
+	}
+
+	ISettingsSectionRef Section = Category->AddSection(SectionName, InDisplayName, InDescription, SettingsObject);
+	CategoryModifiedDelegate.Broadcast(CategoryName);
+
+	return Section;
+}
+
+
+ISettingsSectionPtr FSettingsContainer::AddSection( const FName& CategoryName, const FName& SectionName, const FText& InDisplayName, const FText& InDescription, const TSharedRef<SWidget>& CustomWidget )
+{
+	TSharedPtr<FSettingsCategory> Category = Categories.FindRef(CategoryName);
+
+	if (!Category.IsValid())
+	{
+		return nullptr;
+	}
+
+	ISettingsSectionRef Section = Category->AddSection(SectionName, InDisplayName, InDescription, CustomWidget);
+	CategoryModifiedDelegate.Broadcast(CategoryName);
+
+	return Section;
+}
+
+
+void FSettingsContainer::RemoveSection( const FName& CategoryName, const FName& SectionName )
+{
+	TSharedPtr<FSettingsCategory> Category = Categories.FindRef(CategoryName);
+
+	if (Category.IsValid())
+	{
+		ISettingsSectionPtr Section = Category->GetSection(SectionName);
+
+		if (Section.IsValid())
+		{
+			Category->RemoveSection(SectionName);
+			SectionRemovedDelegate.Broadcast(Section.ToSharedRef());
+			CategoryModifiedDelegate.Broadcast(CategoryName);
+		}
+	}
+}
+
+/* ISettingsContainer interface
+ *****************************************************************************/
+
+void FSettingsContainer::Describe( const FText& InDisplayName, const FText& InDescription, const FName& InIconName )
+{
+	Description = InDescription;
+	DisplayName = InDisplayName;
+	IconName = InIconName;
+}
+
+
+void FSettingsContainer::DescribeCategory( const FName& CategoryName, const FText& InDisplayName, const FText& InDescription )
+{
+	TSharedPtr<FSettingsCategory>& Category = Categories.FindOrAdd(CategoryName);
+
+	if (!Category.IsValid())
+	{
+		Category = MakeShareable(new FSettingsCategory(CategoryName));
+	}
+
+	Category->Describe(InDisplayName, InDescription);
+	CategoryModifiedDelegate.Broadcast(CategoryName);
+}
+
+
+int32 FSettingsContainer::GetCategories( TArray<ISettingsCategoryPtr>& OutCategories ) const
+{
+	OutCategories.Empty(Categories.Num());
+
+	for (TMap<FName, TSharedPtr<FSettingsCategory> >::TConstIterator It(Categories); It; ++It)
+	{
+		OutCategories.Add(It.Value());
+	}
+
+	return OutCategories.Num();
+}
diff --git a/Engine/Source/Developer/Settings/Private/SettingsContainer.h b/Engine/Source/Developer/Settings/Private/SettingsContainer.h
--- a/Engine/Source/Developer/Settings/Private/SettingsContainer.h
+++ b/Engine/Source/Developer/Settings/Private/SettingsContainer.h
@@ -16,9 +16,7 @@ public:
 	 *
 	 * @param InName The container's name.
 	 */
-	FSettingsContainer( const FName& InName )
-		: Name(InName)
-	{ }
+	FSettingsContainer( const FName& InName );
 
 public:
 
@@ -32,23 +30,9 @@ public:
 	 * @param DisplayName The section's localized display name.
 	 * @param Description The section's localized description text.
 	 * @param SettingsObject The object that holds the section's settings.
-	 * @param Delegates The section's optional callback delegates.
 	 * @return The added settings section, or nullptr if the category does not exist.
 	 */
-	ISettingsSectionPtr AddSection( const FName& CategoryName, const FName& SectionName, const FText& InDisplayName, const FText& InDescription, const TWeakObjectPtr<UObject>& SettingsObject, const FSettingsSectionDelegates& Delegates )
-	{
-		TSharedPtr<FSettingsCategory> Category = Categories.FindRef(CategoryName);
-
-		if (!Category.IsValid())
-		{
-			return nullptr;
-		}
-
-		ISettingsSectionRef Section = Category->AddSection(SectionName, InDisplayName, InDescription, SettingsObject, Delegates);
-		CategoryModifiedDelegate.Broadcast(CategoryName);
-
-		return Section;
-	}
+	ISettingsSectionPtr AddSection( const FName& CategoryName, const FName& SectionName, const FText& InDisplayName, const FText& InDescription, const TWeakObjectPtr<UObject>& SettingsObject );
 
 	/**
 	 * Adds a settings section to the specified category (using a custom settings widget).
@@ -60,23 +44,9 @@ public:
 	 * @param DisplayName The section's localized display name.
 	 * @param Description The section's localized description text.
 	 * @param CustomWidget A custom settings widget.
-	 * @param Delegates The section's optional callback delegates.
 	 * @return The added settings section, or nullptr if the category does not exist.
 	 */
-	ISettingsSectionPtr AddSection( const FName& CategoryName, const FName& SectionName, const FText& InDisplayName, const FText& InDescription, const TSharedRef<SWidget>& CustomWidget, const FSettingsSectionDelegates& Delegates )
-	{
-		TSharedPtr<FSettingsCategory> Category = Categories.FindRef(CategoryName);
-
-		if (!Category.IsValid())
-		{
-			return nullptr;
-		}
-
-		ISettingsSectionRef Section = Category->AddSection(SectionName, InDisplayName, InDescription, CustomWidget, Delegates);
-		CategoryModifiedDelegate.Broadcast(CategoryName);
-
-		return Section;
-	}
+	ISettingsSectionPtr AddSection( const FName& CategoryName, const FName& SectionName, const FText& InDisplayName, const FText& InDescription, const TSharedRef<SWidget>& CustomWidget );
 
 	/**
 	 * Removes a settings section.
@@ -84,116 +54,73 @@ public:
 	 * @param CategoryName The name of the category that contains the section.
 	 * @param SectionName The name of the section to remove.
 	 */
-	void RemoveSection( const FName& CategoryName, const FName& SectionName )
-	{
-		TSharedPtr<FSettingsCategory> Category = Categories.FindRef(CategoryName);
-
-		if (Category.IsValid())
-		{
-			ISettingsSectionPtr Section = Category->GetSection(SectionName);
-
-			if (Section.IsValid())
-			{
-				Category->RemoveSection(SectionName);
-				SectionRemovedDelegate.Broadcast(Section.ToSharedRef());
-				CategoryModifiedDelegate.Broadcast(CategoryName);
-			}
-		}
-	}
+	void RemoveSection( const FName& CategoryName, const FName& SectionName );
 
 public:
 
 	// ISettingsContainer interface
 
-	virtual void Describe( const FText& InDisplayName, const FText& InDescription, const FName& InIconName ) override
-	{
-		Description = InDescription;
-		DisplayName = InDisplayName;
-		IconName = InIconName;
-	}
-
-	virtual void DescribeCategory( const FName& CategoryName, const FText& InDisplayName, const FText& InDescription ) override
-	{
-		TSharedPtr<FSettingsCategory>& Category = Categories.FindOrAdd(CategoryName);
-
-		if (!Category.IsValid())
-		{
-			Category = MakeShareable(new FSettingsCategory(CategoryName));
-		}
-
-		Category->Describe(InDisplayName, InDescription);
-		CategoryModifiedDelegate.Broadcast(CategoryName);
-	}
-
-	virtual int32 GetCategories( TArray<ISettingsCategoryPtr>& OutCategories ) const override
-	{
-		OutCategories.Empty(Categories.Num());
-
-		for (TMap<FName, TSharedPtr<FSettingsCategory> >::TConstIterator It(Categories); It; ++It)
-		{
-			OutCategories.Add(It.Value());
-		}
-
-		return OutCategories.Num();
-	}
+	virtual void Describe( const FText& InDisplayName, const FText& InDescription, const FName& InIconName ) override;
+	virtual void DescribeCategory( const FName& CategoryName, const FText& InDisplayName, const FText& InDescription ) override;
+	virtual int32 GetCategories( TArray<TSharedPtr<ISettingsCategory>>& OutCategories ) const override;
 
-	virtual ISettingsCategoryPtr GetCategory( const FName& CategoryName ) const override
+	virtual TSharedPtr<ISettingsCategory> GetCategory( const FName& CategoryName ) const override
 	{
 		return Categories.FindRef(CategoryName);
 	}
 
-	virtual const FText& GetDescription( ) const override
+	virtual const FText& GetDescription() const override
 	{
 		return Description;
 	}
 
-	virtual const FText& GetDisplayName( ) const override
+	virtual const FText& GetDisplayName() const override
 	{
 		return DisplayName;
 	}
 
-	virtual const FName& GetIconName( ) const override
+	virtual const FName& GetIconName() const override
 	{
 		return IconName;
 	}
 
-	virtual const FName& GetName( ) const override
+	virtual const FName& GetName() const override
 	{
 		return Name;
 	}
 
-	virtual FOnSettingsContainerCategoryModified& OnCategoryModified( ) override
+	virtual FOnCategoryModified& OnCategoryModified() override
 	{
 		return CategoryModifiedDelegate;
 	}
 
-	virtual FOnSettingsContainerSectionRemoved& OnSectionRemoved( ) override
+	virtual FOnSectionRemoved& OnSectionRemoved() override
 	{
 		return SectionRemovedDelegate;
 	}
 
 private:
 
-	// Holds the collection of setting categories
+	/** Holds the collection of setting categories. */
 	TMap<FName, TSharedPtr<FSettingsCategory>> Categories;
 
-	// Holds the container's description text.
+	/** Holds the container's description text. */
 	FText Description;
 
-	// Holds the container's localized display name.
+	/** Holds the container's localized display name. */
 	FText DisplayName;
 
-	// Holds the name of the container's icon.
+	/** Holds the name of the container's icon. */
 	FName IconName;
 
-	// Holds the container's name.
+	/** Holds the container's name. */
 	FName Name;
 
 private:
 
-	// Holds a delegate that is executed when a settings category has been added or modified.
-	FOnSettingsContainerCategoryModified CategoryModifiedDelegate;
+	/** Holds a delegate that is executed when a settings category has been added or modified. */
+	FOnCategoryModified CategoryModifiedDelegate;
 
-	// Holds a delegate that is executed when a settings section has been removed.
-	FOnSettingsContainerSectionRemoved SectionRemovedDelegate;
+	/** Holds a delegate that is executed when a settings section has been removed. */
+	FOnSectionRemoved SectionRemovedDelegate;
 };
diff --git a/Engine/Source/Developer/Settings/Private/SettingsModule.cpp b/Engine/Source/Developer/Settings/Private/SettingsModule.cpp
--- a/Engine/Source/Developer/Settings/Private/SettingsModule.cpp
+++ b/Engine/Source/Developer/Settings/Private/SettingsModule.cpp
@@ -1,6 +1,7 @@
 // Copyright 1998-2014 Epic Games, Inc. All Rights Reserved.
 
 #include "SettingsPrivatePCH.h"
+#include "ISettingsModule.h"
 
 
 #define LOCTEXT_NAMESPACE "FSettingsModule"
@@ -26,14 +27,14 @@ public:
 		return ContainerNamesToContainers.FindRef(ContainerName);
 	}
 
-	virtual ISettingsSectionPtr RegisterSettings( const FName& ContainerName, const FName& CategoryName, const FName& SectionName, const FText& DisplayName, const FText& Description, const TWeakObjectPtr<UObject>& SettingsObject, const FSettingsSectionDelegates& Delegates ) override
+	virtual ISettingsSectionPtr RegisterSettings( const FName& ContainerName, const FName& CategoryName, const FName& SectionName, const FText& DisplayName, const FText& Description, const TWeakObjectPtr<UObject>& SettingsObject ) override
 	{
-		return FindOrAddContainer(ContainerName)->AddSection(CategoryName, SectionName, DisplayName, Description, SettingsObject, Delegates);
+		return FindOrAddContainer(ContainerName)->AddSection(CategoryName, SectionName, DisplayName, Description, SettingsObject);
 	}
 
-	virtual ISettingsSectionPtr RegisterSettings( const FName& ContainerName, const FName& CategoryName, const FName& SectionName, const FText& DisplayName, const FText& Description, const TSharedRef<SWidget>& CustomWidget, const FSettingsSectionDelegates& Delegates ) override
+	virtual ISettingsSectionPtr RegisterSettings( const FName& ContainerName, const FName& CategoryName, const FName& SectionName, const FText& DisplayName, const FText& Description, const TSharedRef<SWidget>& CustomWidget ) override
 	{
-		return FindOrAddContainer(ContainerName)->AddSection(CategoryName, SectionName, DisplayName, Description, CustomWidget, Delegates);
+		return FindOrAddContainer(ContainerName)->AddSection(CategoryName, SectionName, DisplayName, Description, CustomWidget);
 	}
 
 	virtual void RegisterViewer( const FName& ContainerName, ISettingsViewer& SettingsViewer )
@@ -70,7 +71,7 @@ public:
 
 	// IModuleInterface interface
 	
-	virtual void StartupModule( ) override
+	virtual void StartupModule() override
 	{
 		// @todo gmp: move this into the modules that own these setting categories
 		TSharedRef<FSettingsContainer> EditorSettingsContainer = FindOrAddContainer("Editor");
@@ -89,7 +90,7 @@ public:
 		ProjectSettingsContainer->DescribeCategory("Plugins", LOCTEXT("ProjectPluginsCategoryName", "Plugins"), LOCTEXT("ProjectPluginsCategoryDescription", "Plugins settings description text here"));
 	}
 
-	virtual void ShutdownModule( ) override { }
+	virtual void ShutdownModule() override { }
 
 protected:
 
@@ -113,13 +114,13 @@ protected:
 
 private:
 
-	// Holds the default settings container.
+	/** Holds the default settings container. */
 	ISettingsContainerPtr DefaultSettingsContainer;
 
-	// Holds the collection of global settings containers.
+	/** Holds the collection of global settings containers. */
 	TMap<FName, TSharedPtr<FSettingsContainer> > ContainerNamesToContainers;
 
-	// Holds the collection of registered settings viewers.
+	/** Holds the collection of registered settings viewers. */
 	TMap<FName, ISettingsViewer*> ContainerNamesToViewers;
 };
 
diff --git a/Engine/Source/Developer/Settings/Private/SettingsPrivatePCH.h b/Engine/Source/Developer/Settings/Private/SettingsPrivatePCH.h
--- a/Engine/Source/Developer/Settings/Private/SettingsPrivatePCH.h
+++ b/Engine/Source/Developer/Settings/Private/SettingsPrivatePCH.h
@@ -3,13 +3,15 @@
 #pragma once
 
 
-#include "Settings.h"
-
-
-/* Dependencies
+/* Private dependencies
  *****************************************************************************/
 
-#include "ModuleManager.h"
+#include "Core.h"
+#include "CoreUObject.h"
+#include "ISettingsCategory.h"
+#include "ISettingsContainer.h"
+#include "ISettingsSection.h"
+#include "ISettingsViewer.h"
 
 
 /* Private includes
@@ -17,4 +19,4 @@
 
 #include "SettingsSection.h"
 #include "SettingsCategory.h"
-#include "SettingsContainer.h"
\ No newline at end of file
+#include "SettingsContainer.h"
diff --git a/Engine/Source/Developer/Settings/Private/SettingsSection.cpp b/Engine/Source/Developer/Settings/Private/SettingsSection.cpp
new file mode 100644
--- /dev/null
+++ b/Engine/Source/Developer/Settings/Private/SettingsSection.cpp
@@ -0,0 +1,236 @@
+// Copyright 1998-2014 Epic Games, Inc. All Rights Reserved.
+
+#include "SettingsPrivatePCH.h"
+
+
+/* FSettingsSection structors
+ *****************************************************************************/
+
+FSettingsSection::FSettingsSection( const ISettingsCategoryRef& InCategory, const FName& InName, const FText& InDisplayName, const FText& InDescription, const TWeakObjectPtr<UObject>& InSettingsObject )
+	: Category(InCategory)
+	, Description(InDescription)
+	, DisplayName(InDisplayName)
+	, Name(InName)
+	, SettingsObject(InSettingsObject)
+{ }
+
+
+FSettingsSection::FSettingsSection( const ISettingsCategoryRef& InCategory, const FName& InName, const FText& InDisplayName, const FText& InDescription, const TSharedRef<SWidget>& InCustomWidget )
+	: Category(InCategory)
+	, CustomWidget(InCustomWidget)
+	, Description(InDescription)
+	, DisplayName(InDisplayName)
+	, Name(InName)
+{ }
+
+
+/* ISettingsSection interface
+ *****************************************************************************/
+
+bool FSettingsSection::CanEdit() const
+{
+	if (CanEditDelegate.IsBound())
+	{
+		return CanEditDelegate.Execute();
+	}
+
+	return true;
+}
+
+
+bool FSettingsSection::CanExport() const
+{
+	return (ExportDelegate.IsBound() || (SettingsObject.IsValid() && SettingsObject->GetClass()->HasAnyClassFlags(CLASS_Config)));
+}
+
+
+bool FSettingsSection::CanImport() const
+{
+	return (ImportDelegate.IsBound() || (SettingsObject.IsValid() && SettingsObject->GetClass()->HasAnyClassFlags(CLASS_Config)));
+}
+
+
+bool FSettingsSection::CanResetDefaults() const
+{
+	return (ResetDefaultsDelegate.IsBound() || (SettingsObject.IsValid() && SettingsObject->GetClass()->HasAnyClassFlags(CLASS_Config) && !SettingsObject->GetClass()->HasAnyClassFlags(CLASS_DefaultConfig)));
+}
+
+
+bool FSettingsSection::CanSave() const
+{
+	return (SaveDelegate.IsBound() || (SettingsObject.IsValid() && SettingsObject->GetClass()->HasAnyClassFlags(CLASS_Config)));
+}
+
+
+bool FSettingsSection::CanSaveDefaults() const
+{
+	return (SaveDefaultsDelegate.IsBound() || (SettingsObject.IsValid() && SettingsObject->GetClass()->HasAnyClassFlags(CLASS_Config) && !SettingsObject->GetClass()->HasAnyClassFlags(CLASS_DefaultConfig)));
+}
+
+
+bool FSettingsSection::Export( const FString& Filename )
+{
+	if (ExportDelegate.IsBound())
+	{
+		return ExportDelegate.Execute(Filename);
+	}
+
+	if (SettingsObject.IsValid())
+	{
+		SettingsObject->SaveConfig(CPF_Config, *Filename);
+
+		return true;
+	}
+
+	return false;
+}
+
+
+TWeakPtr<ISettingsCategory> FSettingsSection::GetCategory()
+{
+	return Category;
+}
+
+
+TWeakPtr<SWidget> FSettingsSection::GetCustomWidget() const
+{
+	return CustomWidget;
+}
+
+
+const FText& FSettingsSection::GetDescription() const
+{
+	return Description;
+}
+
+
+const FText& FSettingsSection::GetDisplayName() const
+{
+	return DisplayName;
+}
+
+
+const FName& FSettingsSection::GetName() const
+{
+	return Name;
+}
+
+
+TWeakObjectPtr<UObject> FSettingsSection::GetSettingsObject() const
+{
+	return SettingsObject;
+}
+
+
+FText FSettingsSection::GetStatus() const
+{
+	if (StatusDelegate.IsBound())
+	{
+		return StatusDelegate.Execute();
+	}
+
+	return FText::GetEmpty();
+}
+
+
+bool FSettingsSection::HasDefaultSettingsObject()
+{
+	if (!SettingsObject.IsValid())
+	{
+		return false;
+	}
+
+	return SettingsObject->GetClass()->HasAnyClassFlags(CLASS_DefaultConfig);
+}
+
+
+bool FSettingsSection::Import( const FString& Filename )
+{
+	if (ImportDelegate.IsBound())
+	{
+		return ImportDelegate.Execute(Filename);
+	}
+
+	if (SettingsObject.IsValid())
+	{
+		SettingsObject->LoadConfig(SettingsObject->GetClass(), *Filename, UE4::LCPF_PropagateToInstances);
+
+		return true;
+	}
+
+	return false;	
+}
+
+
+bool FSettingsSection::ResetDefaults()
+{
+	if (ResetDefaultsDelegate.IsBound())
+	{
+		return ResetDefaultsDelegate.Execute();
+	}
+
+	if (SettingsObject.IsValid() && SettingsObject->GetClass()->HasAnyClassFlags(CLASS_Config) && !SettingsObject->GetClass()->HasAnyClassFlags(CLASS_DefaultConfig))
+	{
+		FString ConfigName = SettingsObject->GetClass()->GetConfigName();
+
+		GConfig->EmptySection(*SettingsObject->GetClass()->GetPathName(), ConfigName);
+		GConfig->Flush(false);
+
+		FConfigCacheIni::LoadGlobalIniFile(ConfigName, *FPaths::GetBaseFilename(ConfigName), nullptr, nullptr, true);
+
+		SettingsObject->ReloadConfig(nullptr, nullptr, UE4::LCPF_PropagateToInstances|UE4::LCPF_PropagateToChildDefaultObjects);
+
+		return true;
+	}
+
+	return false;
+}
+
+
+bool FSettingsSection::Save()
+{
+	if (ModifiedDelegate.IsBound() && !ModifiedDelegate.Execute())
+	{
+		return false;
+	}
+
+	if (SaveDelegate.IsBound())
+	{
+		return SaveDelegate.Execute();
+	}
+
+	if (SettingsObject.IsValid())
+	{
+		if (SettingsObject->GetClass()->HasAnyClassFlags(CLASS_DefaultConfig))
+		{
+			SettingsObject->UpdateDefaultConfigFile();
+		}
+		else
+		{
+			SettingsObject->SaveConfig();
+		}
+
+		return true;
+	}
+
+	return false;
+}
+
+
+bool FSettingsSection::SaveDefaults()
+{
+	if (SaveDefaultsDelegate.IsBound())
+	{
+		return SaveDefaultsDelegate.Execute();
+	}
+
+	if (SettingsObject.IsValid())
+	{
+		SettingsObject->UpdateDefaultConfigFile();
+		SettingsObject->ReloadConfig(nullptr, nullptr, UE4::LCPF_PropagateToInstances);
+
+		return true;			
+	}
+
+	return false;
+}
diff --git a/Engine/Source/Developer/Settings/Private/SettingsSection.h b/Engine/Source/Developer/Settings/Private/SettingsSection.h
--- a/Engine/Source/Developer/Settings/Private/SettingsSection.h
+++ b/Engine/Source/Developer/Settings/Private/SettingsSection.h
@@ -18,24 +18,9 @@ public:
 	 * @param InName The setting section's name.
 	 * @param InDisplayName The section's localized display name.
 	 * @param InDescription The section's localized description text.
-	 * @param InDelegates The section's optional callback delegates.
 	 * @param InSettingsObject The object that holds the settings for this section.
 	 */
-	FSettingsSection( const ISettingsCategoryRef& InCategory, const FName& InName, const FText& InDisplayName, const FText& InDescription, const FSettingsSectionDelegates& InDelegates, const TWeakObjectPtr<UObject>& InSettingsObject )
-		: Category(InCategory)
-		, Description(InDescription)
-		, DisplayName(InDisplayName)
-		, Name(InName)
-		, SettingsObject(InSettingsObject)
-		, CanEditDelegate(InDelegates.CanEditDelegate)
-		, ExportDelegate(InDelegates.ExportDelegate)
-		, ImportDelegate(InDelegates.ImportDelegate)
-		, ModifiedDelegate(InDelegates.ModifiedDelegate)
-		, ResetDefaultsDelegate(InDelegates.ResetDefaultsDelegate)
-		, SaveDefaultsDelegate(InDelegates.SaveDefaultsDelegate)
-		, SaveDelegate(InDelegates.SaveDelegate)
-		, StatusDelegate(InDelegates.StatusDelegate)
-	{ }
+	FSettingsSection( const ISettingsCategoryRef& InCategory, const FName& InName, const FText& InDisplayName, const FText& InDescription, const TWeakObjectPtr<UObject>& InSettingsObject );
 
 	/**
 	 * Creates and initializes a new settings section from the given custom settings widget.
@@ -44,262 +29,118 @@ public:
 	 * @param InName The setting section's name.
 	 * @param InDisplayName The section's localized display name.
 	 * @param InDescription The section's localized description text.
-	 * @param InDelegates The section's optional callback delegates.
 	 * @param InCustomWidget A custom settings widget.
 	 */
-	FSettingsSection( const ISettingsCategoryRef& InCategory, const FName& InName, const FText& InDisplayName, const FText& InDescription, const FSettingsSectionDelegates& InDelegates, const TSharedRef<SWidget>& InCustomWidget )
-		: Category(InCategory)
-		, CustomWidget(InCustomWidget)
-		, Description(InDescription)
-		, DisplayName(InDisplayName)
-		, Name(InName)
-		, CanEditDelegate(InDelegates.CanEditDelegate)
-		, ExportDelegate(InDelegates.ExportDelegate)
-		, ImportDelegate(InDelegates.ImportDelegate)
-		, ModifiedDelegate(InDelegates.ModifiedDelegate)
-		, ResetDefaultsDelegate(InDelegates.ResetDefaultsDelegate)
-		, SaveDefaultsDelegate(InDelegates.SaveDefaultsDelegate)
-		, SaveDelegate(InDelegates.SaveDelegate)
-		, StatusDelegate(InDelegates.StatusDelegate)
-	{ }
+	FSettingsSection( const ISettingsCategoryRef& InCategory, const FName& InName, const FText& InDisplayName, const FText& InDescription, const TSharedRef<SWidget>& InCustomWidget );
 
 public:
 
 	// ISettingsSection interface
 
-	virtual bool CanEdit( ) const override
-	{
-		if (CanEditDelegate.IsBound())
-		{
-			return CanEditDelegate.Execute();
-		}
-
-		return true;
-	}
+	virtual bool CanEdit() const override;
+	virtual bool CanExport() const override;
+	virtual bool CanImport() const override;
+	virtual bool CanResetDefaults() const override;
+	virtual bool CanSave() const override;
+	virtual bool CanSaveDefaults() const override;
+	virtual bool Export( const FString& Filename ) override;
+	virtual TWeakPtr<ISettingsCategory> GetCategory() override;
+	virtual TWeakPtr<SWidget> GetCustomWidget() const override;
+	virtual const FText& GetDescription() const override;
+	virtual const FText& GetDisplayName() const override;
+	virtual const FName& GetName() const override;
+	virtual TWeakObjectPtr<UObject> GetSettingsObject() const override;
+	virtual FText GetStatus() const override;
+	virtual bool HasDefaultSettingsObject() override;
+	virtual bool Import( const FString& Filename ) override;
 
-	virtual bool CanExport( ) const override
+	virtual FOnCanEdit& OnCanEdit() override
 	{
-		return (ExportDelegate.IsBound() || (SettingsObject.IsValid() && SettingsObject->GetClass()->HasAnyClassFlags(CLASS_Config)));
+		return CanEditDelegate;
 	}
 
-	virtual bool CanImport( ) const override
+	virtual FOnExport& OnExport() override
 	{
-		return (ImportDelegate.IsBound() || (SettingsObject.IsValid() && SettingsObject->GetClass()->HasAnyClassFlags(CLASS_Config)));
+		return ExportDelegate;
 	}
 
-	virtual bool CanResetDefaults( ) const override
+	virtual FOnImport& OnImport() override
 	{
-		return (ResetDefaultsDelegate.IsBound() || (SettingsObject.IsValid() && SettingsObject->GetClass()->HasAnyClassFlags(CLASS_Config) && !SettingsObject->GetClass()->HasAnyClassFlags(CLASS_DefaultConfig)));
+		return ImportDelegate;
 	}
 
-	virtual bool CanSave( ) const override
+	virtual FOnModified& OnModified() override
 	{
-		return (SaveDelegate.IsBound() || (SettingsObject.IsValid() && SettingsObject->GetClass()->HasAnyClassFlags(CLASS_Config)));
+		return ModifiedDelegate;
 	}
 
-	virtual bool CanSaveDefaults( ) const override
+	virtual FOnResetDefaults& OnResetDefaults() override
 	{
-		return (SaveDefaultsDelegate.IsBound() || (SettingsObject.IsValid() && SettingsObject->GetClass()->HasAnyClassFlags(CLASS_Config) && !SettingsObject->GetClass()->HasAnyClassFlags(CLASS_DefaultConfig)));
+		return ResetDefaultsDelegate;
 	}
 
-	virtual bool Export( const FString& Filename ) override
+	virtual FOnSave& OnSave() override
 	{
-		if (ExportDelegate.IsBound())
-		{
-			return ExportDelegate.Execute(Filename);
-		}
-
-		if (SettingsObject.IsValid())
-		{
-			SettingsObject->SaveConfig(CPF_Config, *Filename);
-
-			return true;
-		}
-
-		return false;
+		return SaveDelegate;
 	}
 
-	virtual ISettingsCategoryWeakPtr GetCategory( ) override
+	virtual FOnSaveDefaults& OnSaveDefaults() override
 	{
-		return Category;
+		return SaveDefaultsDelegate;
 	}
 
-	virtual TWeakPtr<SWidget> GetCustomWidget( ) const override
+	virtual FOnStatus& OnStatus() override
 	{
-		return CustomWidget;
+		return StatusDelegate;
 	}
 
-	virtual const FText& GetDescription( ) const override
-	{
-		return Description;
-	}
-
-	virtual const FText& GetDisplayName( ) const override
-	{
-		return DisplayName;
-	}
-
-	virtual const FName& GetName( ) const override
-	{
-		return Name;
-	}
-
-	virtual TWeakObjectPtr<UObject> GetSettingsObject( ) const override
-	{
-		return SettingsObject;
-	}
-
-	virtual FText GetStatus( ) const override
-	{
-		if (StatusDelegate.IsBound())
-		{
-			return StatusDelegate.Execute();
-		}
-
-		return FText::GetEmpty();
-	}
-
-	virtual bool HasDefaultSettingsObject( ) override
-	{
-		if (!SettingsObject.IsValid())
-		{
-			return false;
-		}
-
-		return SettingsObject->GetClass()->HasAnyClassFlags(CLASS_DefaultConfig);
-	}
-
-	virtual bool Import( const FString& Filename ) override
-	{
-		if (ImportDelegate.IsBound())
-		{
-			return ImportDelegate.Execute(Filename);
-		}
-
-		if (SettingsObject.IsValid())
-		{
-			SettingsObject->LoadConfig(SettingsObject->GetClass(), *Filename, UE4::LCPF_PropagateToInstances);
-
-			return true;
-		}
-
-		return false;	
-	}
-
-	virtual bool ResetDefaults( ) override
-	{
-		if (ResetDefaultsDelegate.IsBound())
-		{
-			return ResetDefaultsDelegate.Execute();
-		}
-
-		if (SettingsObject.IsValid() && SettingsObject->GetClass()->HasAnyClassFlags(CLASS_Config) && !SettingsObject->GetClass()->HasAnyClassFlags(CLASS_DefaultConfig))
-		{
-			FString ConfigName = SettingsObject->GetClass()->GetConfigName();
-
-			GConfig->EmptySection(*SettingsObject->GetClass()->GetPathName(), ConfigName);
-			GConfig->Flush(false);
-
-			FConfigCacheIni::LoadGlobalIniFile(ConfigName, *FPaths::GetBaseFilename(ConfigName), nullptr, nullptr, true);
-
-			SettingsObject->ReloadConfig(nullptr, nullptr, UE4::LCPF_PropagateToInstances|UE4::LCPF_PropagateToChildDefaultObjects);
-
-			return true;
-		}
-
-		return false;
-	}
-
-	virtual bool Save( ) override
-	{
-		if (ModifiedDelegate.IsBound() && !ModifiedDelegate.Execute())
-		{
-			return false;
-		}
-
-		if (SaveDelegate.IsBound())
-		{
-			return SaveDelegate.Execute();
-		}
-
-		if (SettingsObject.IsValid())
-		{
-			if (SettingsObject->GetClass()->HasAnyClassFlags(CLASS_DefaultConfig))
-			{
-				SettingsObject->UpdateDefaultConfigFile();
-			}
-			else
-			{
-				SettingsObject->SaveConfig();
-			}
-
-			return true;
-		}
-
-		return false;
-	}
-
-	virtual bool SaveDefaults( ) override
-	{
-		if (SaveDefaultsDelegate.IsBound())
-		{
-			return SaveDefaultsDelegate.Execute();
-		}
-
-		if (SettingsObject.IsValid())
-		{
-			SettingsObject->UpdateDefaultConfigFile();
-			SettingsObject->ReloadConfig(nullptr, nullptr, UE4::LCPF_PropagateToInstances);
-
-			return true;			
-		}
-
-		return false;
-	}
+	virtual bool ResetDefaults() override;
+	virtual bool Save() override;
+	virtual bool SaveDefaults() override;
 
 private:
 
-	// Holds a pointer to the settings category that owns this section.
-	ISettingsCategoryWeakPtr Category;
+	/** Holds a pointer to the settings category that owns this section. */
+	TWeakPtr<ISettingsCategory> Category;
 
-	// Holds the section's custom editor widget.
+	/** Holds the section's custom editor widget. */
 	TWeakPtr<SWidget> CustomWidget;
 
-	// Holds the section's description text.
+	/** Holds the section's description text. */
 	FText Description;
 
-	// Holds the section's localized display name.
+	/** Holds the section's localized display name. */
 	FText DisplayName;
 
-	// Holds the section's name.
+	/** Holds the section's name. */
 	FName Name;
 
-	// Holds the settings object.
+	/** Holds the settings object. */
 	TWeakObjectPtr<UObject> SettingsObject;
 
 private:
 
-	// Holds a delegate that is executed to check whether a settings section can be edited.
-	FOnSettingsSectionCanEdit CanEditDelegate;
+	/** Holds a delegate that is executed to check whether a settings section can be edited. */
+	FOnCanEdit CanEditDelegate;
 
-	// Holds a delegate that is executed when the settings section should be exported to a file.
-	FOnSettingsSectionExport ExportDelegate;
+	/** Holds a delegate that is executed when the settings section should be exported to a file. */
+	FOnExport ExportDelegate;
 
-	// Holds a delegate that is executed when the settings section should be imported from a file.
-	FOnSettingsSectionImport ImportDelegate;
+	/** Holds a delegate that is executed when the settings section should be imported from a file. */
+	FOnImport ImportDelegate;
 
-	// Holds a delegate that is executed after the settings section has been modified.
-	FOnSettingsSectionModified ModifiedDelegate;
+	/** Holds a delegate that is executed after the settings section has been modified. */
+	FOnModified ModifiedDelegate;
 
-	// Holds a delegate that is executed when the settings section should have its values reset to default.
-	FOnSettingsSectionResetDefaults ResetDefaultsDelegate;
+	/** Holds a delegate that is executed when the settings section should have its values reset to default. */
+	FOnResetDefaults ResetDefaultsDelegate;
 
-	// Holds a delegate that is executed when a settings section should have its values saved as default.
-	FOnSettingsSectionSaveDefaults SaveDefaultsDelegate;
+	/** Holds a delegate that is executed when a settings section should have its values saved as default. */
+	FOnSaveDefaults SaveDefaultsDelegate;
 
-	// Holds a delegate that is executed when a settings section should have its values saved.
-	FOnSettingsSectionSave SaveDelegate;
+	/** Holds a delegate that is executed when a settings section should have its values saved. */
+	FOnSave SaveDelegate;
 
-	// Holds a delegate that is executed to retrieve a status message for a settings section.
-	FOnSettingsSectionStatus StatusDelegate;
+	/** Holds a delegate that is executed to retrieve a status message for a settings section. */
+	FOnStatus StatusDelegate;
 };
diff --git a/Engine/Source/Developer/Settings/Public/ISettingsCategory.h b/Engine/Source/Developer/Settings/Public/ISettingsCategory.h
new file mode 100644
--- /dev/null
+++ b/Engine/Source/Developer/Settings/Public/ISettingsCategory.h
@@ -0,0 +1,69 @@
+// Copyright 1998-2014 Epic Games, Inc. All Rights Reserved.
+
+#pragma once
+
+
+// forward declarations
+class ISettingsSection;
+
+
+/**
+ * Interface for setting categories.
+ *
+ * A setting category is a collection of setting sections that logically belong together
+ * (i.e. all project settings or all platform settings). Each section then contains the
+ * actual settings in the form of UObject properties.
+ */
+class ISettingsCategory
+{
+public:
+
+	/**
+	 * Gets the category's localized description text.
+	 *
+	 * @return Description text.
+	 */
+	virtual const FText& GetDescription() const = 0;
+
+	/**
+	 * Gets the category's localized display name.
+	 *
+	 * @return Display name.
+	 */
+	virtual const FText& GetDisplayName() const = 0;
+
+	/**
+	 * Gets the settings section with the specified name.
+	 *
+	 * @param SectionName The name of the section to get.
+	 * @return The settings section, or nullptr if it doesn't exist.
+	 */
+	virtual TSharedPtr<ISettingsSection> GetSection( const FName& SectionName ) const = 0;
+
+	/**
+	 * Gets the setting sections contained in this category.
+	 *
+	 * @param OutSections Will hold the collection of sections.
+	 * @return The number of sections returned.
+	 */
+	virtual int32 GetSections( TArray<TSharedPtr<ISettingsSection>>& OutSections ) const = 0;
+
+	/**
+	 * Gets the category's name.
+	 *
+	 * @return Category name.
+	 */
+	virtual const FName& GetName() const = 0;
+
+public:
+
+	/** Virtual destructor. */
+	virtual ~ISettingsCategory() { }
+};
+
+
+/** Type definition for shared pointers to instances of ISettingsCategory. */
+typedef TSharedPtr<ISettingsCategory> ISettingsCategoryPtr;
+
+/** Type definition for shared references to instances of ISettingsCategory. */
+typedef TSharedRef<ISettingsCategory> ISettingsCategoryRef;
diff --git a/Engine/Source/Developer/Settings/Public/ISettingsContainer.h b/Engine/Source/Developer/Settings/Public/ISettingsContainer.h
new file mode 100644
--- /dev/null
+++ b/Engine/Source/Developer/Settings/Public/ISettingsContainer.h
@@ -0,0 +1,112 @@
+// Copyright 1998-2014 Epic Games, Inc. All Rights Reserved.
+
+#pragma once
+
+
+// forward declarations
+class ISettingsCategory;
+class ISettingsSection;
+
+
+/**
+ * Interface for setting containers.
+ *
+ * A settings container is a collection of setting categories.
+ * Each category holds a collection of setting sections, which contain
+ * the actual settings in the form of UObject properties.
+ */
+class ISettingsContainer
+{
+public:
+
+	/**
+	 * Updates the details of this settings container.
+	 *
+	 * @param InDisplayName The container's localized display name.
+	 * @param InDescription The container's localized description text.
+	 * @param InIconName The name of the container's icon.
+	 */
+	virtual void Describe( const FText& InDisplayName, const FText& InDescription, const FName& InIconName ) = 0;
+
+	/**
+	 * Updates the details of the specified settings category.
+	 *
+	 * @param CategoryName The name of the category to update.
+	 * @param DisplayName The category's localized display name.
+	 * @param Description The category's localized description text.
+	 */
+	virtual void DescribeCategory( const FName& CategoryName, const FText& DisplayName, const FText& Description ) = 0;
+
+	/**
+	 * Gets the setting categories.
+	 *
+	 * @param OutCategories Will hold the collection of categories.
+	 * @return The number of categories returned.
+	 */
+	virtual int32 GetCategories( TArray<TSharedPtr<ISettingsCategory>>& OutCategories ) const = 0;
+
+	/**
+	 * Gets the category with the specified name.
+	 *
+	 * @return The category, or nullptr if it doesn't exist.
+	 */
+	virtual TSharedPtr<ISettingsCategory> GetCategory( const FName& CategoryName ) const = 0;
+
+	/**
+	 * Gets the container's localized description text.
+	 *
+	 * @return Description text.
+	 */
+	virtual const FText& GetDescription() const = 0;
+
+	/**
+	 * Gets the container's localized display name.
+	 *
+	 * @return Display name.
+	 */
+	virtual const FText& GetDisplayName() const = 0;
+
+	/**
+	 * Gets the name of the container's icon.
+	 *
+	 * @return Icon image name.
+	 */
+	virtual const FName& GetIconName() const = 0;
+
+	/**
+	 * Gets the container's name.
+	 *
+	 * @return Container name.
+	 */
+	virtual const FName& GetName() const = 0;
+
+public:
+
+	/**
+	 * A delegate that is executed when a settings category has been added or modified.
+	 *
+	 * @return The delegate.
+	 */
+	DECLARE_MULTICAST_DELEGATE_OneParam(FOnCategoryModified, const FName&)
+	virtual FOnCategoryModified& OnCategoryModified() = 0;
+
+	/**
+	 * A delegate that is executed when a settings section has been removed.
+	 *
+	 * @return The delegate.
+	 */
+	DECLARE_MULTICAST_DELEGATE_OneParam(FOnSectionRemoved, const TSharedRef<ISettingsSection>&)
+	virtual FOnSectionRemoved& OnSectionRemoved() = 0;
+
+public:
+
+	/** Virtual destructor. */
+	virtual ~ISettingsContainer() { }
+};
+
+
+/** Type definition for shared pointers to instances of ISettingsContainer. */
+typedef TSharedPtr<ISettingsContainer> ISettingsContainerPtr;
+
+/** Type definition for shared references to instances of ISettingsContainer. */
+typedef TSharedRef<ISettingsContainer> ISettingsContainerRef;
diff --git a/Engine/Source/Developer/Settings/Public/ISettingsModule.h b/Engine/Source/Developer/Settings/Public/ISettingsModule.h
new file mode 100644
--- /dev/null
+++ b/Engine/Source/Developer/Settings/Public/ISettingsModule.h
@@ -0,0 +1,106 @@
+// Copyright 1998-2014 Epic Games, Inc. All Rights Reserved.
+
+#pragma once
+
+#include "ModuleInterface.h"
+
+
+// forward declarations
+class ISettingsContainer;
+class ISettingsSection;
+class ISettingsViewer;
+class SWidget;
+class UObject;
+
+
+/**
+ * Interface for settings UI modules.
+ */
+class ISettingsModule
+	: public IModuleInterface
+{
+public:
+
+	/**
+	 * Gets the names of all known setting containers.
+	 *
+	 * @param OutNames Will contain the collection of names.
+	 */
+	virtual void GetContainerNames( TArray<FName>& OutNames ) const = 0;
+
+	/**
+	 * Gets the global settings container with the specified name.
+	 *
+	 * @param ContainerName The name of the container to get.
+	 * @return The settings container, or nullptr if it doesn't exist.
+	 */
+	virtual TSharedPtr<ISettingsContainer> GetContainer( const FName& ContainerName ) = 0;
+
+	/**
+	 * Adds a settings section to the specified settings container (using a settings object).
+	 *
+	 * If a section with the specified settings objects already exists, the existing section will be replaced.
+	 *
+	 * @param ContainerName The name of the container that will contain the settings.
+	 * @param CategoryName The name of the category within the container.
+	 * @param SectionName The name of the section within the category.
+	 * @param DisplayName The section's localized display name.
+	 * @param Description The section's localized description text.
+	 * @param SettingsObject The object that holds the section's settings.
+	 * @return The added settings section, or nullptr if the category does not exist.
+	 */
+	virtual TSharedPtr<ISettingsSection> RegisterSettings( const FName& ContainerName, const FName& CategoryName, const FName& SectionName, const FText& DisplayName, const FText& Description, const TWeakObjectPtr<UObject>& SettingsObject ) = 0;
+
+	/**
+	 * Adds a settings section to the specified settings container (using a custom settings widget).
+	 *
+	 * If a section with the specified settings objects already exists, the existing section will be replaced.
+	 *
+	 * @param ContainerName The name of the container that will contain the settings.
+	 * @param CategoryName The name of the category within the container.
+	 * @param SectionName The name of the section within the category.
+	 * @param DisplayName The section's localized display name.
+	 * @param Description The section's localized description text.
+	 * @param CustomWidget A custom settings widget.
+	 * @return The added settings section, or nullptr if the category does not exist.
+	 */
+	virtual TSharedPtr<ISettingsSection> RegisterSettings( const FName& ContainerName, const FName& CategoryName, const FName& SectionName, const FText& DisplayName, const FText& Description, const TSharedRef<SWidget>& CustomWidget ) = 0;
+
+	/**
+	 * Registers a viewer for the specified settings container.
+	 *
+	 * @param ContainerName The name of the settings container to register a viewer for.
+	 * @param SettingsViewer The viewer to register.
+	 */
+	virtual void RegisterViewer( const FName& ContainerName, ISettingsViewer& SettingsViewer ) = 0;
+
+	/**
+	 * Shows the settings viewer for the specified settings container.
+	 *
+	 * @param ContainerName The name of the section's container.
+	 * @param CategoryName The name of the section's category.
+	 * @param SectionName The name of the section to show.
+	 */
+	virtual void ShowViewer( const FName& ContainerName, const FName& CategoryName, const FName& SectionName ) = 0;
+
+	/**
+	 * Removes a settings section from the specified settings container.
+	 *
+	 * @param ContainerName The name of the container that to remove the settings from.
+	 * @param CategoryName The name of the category within the container.
+	 * @param SectionName The name of the section within the category.
+	 */
+	virtual void UnregisterSettings( const FName& ContainerName, const FName& CategoryName, const FName& SectionName ) = 0;
+
+	/**
+	 * Unregisters the currently assigned viewer for the specified settings container.
+	 *
+	 * @param ContainerName The name of the settings container to unregister the viewer for.
+	 */
+	virtual void UnregisterViewer( const FName& ContainerName ) = 0;
+
+public:
+
+	/** Virtual destructor. */
+	virtual ~ISettingsModule() { }
+};
diff --git a/Engine/Source/Developer/Settings/Public/ISettingsSection.h b/Engine/Source/Developer/Settings/Public/ISettingsSection.h
new file mode 100644
--- /dev/null
+++ b/Engine/Source/Developer/Settings/Public/ISettingsSection.h
@@ -0,0 +1,241 @@
+// Copyright 1998-2014 Epic Games, Inc. All Rights Reserved.
+
+#pragma once
+
+
+// forward declarations
+class ISettingsCategory;
+class SWidget;
+class UObject;
+
+
+/**
+ * Interface for setting sections.
+ *
+ * A setting section is a collection of settings that logically belong together
+ * (i.e. all settings for the PS4 Platform). Internally, the individual settings
+ * are represented as the properties of an UObject. One or more setting sections
+ * can the be hierarchically arranged in a setting category.
+ */
+class ISettingsSection
+{
+public:
+
+	/**
+	 * Checks whether this section can be edited right now.
+	 *
+	 * @return true if the section can be edited, false otherwise.
+	 */
+	virtual bool CanEdit() const = 0;
+
+	/**
+	 * Checks whether this section can export its settings to a file.
+	 *
+	 * @return true if file export is supported, false otherwise.
+	 */
+	virtual bool CanExport() const = 0;
+
+	/**
+	 * Checks whether this section can import its settings from a file.
+	 *
+	 * @return true if file import is supported, false otherwise.
+	 */
+	virtual bool CanImport() const = 0;
+
+	/**
+	 * Checks whether this section can have its settings reset to default.
+	 *
+	 * @return true if resetting to default is supported, false otherwise.
+	 */
+	virtual bool CanResetDefaults() const = 0;
+
+	/**
+	 * Checks whether this section can have its settings saved.
+	 *
+	 * This method does not indicate whether saving would actually succeed.
+	 * For example, saving may be aborted by a ModifiedDelegate handler.
+	 *
+	 * @return true if saving is supported, false otherwise.
+	 */
+	virtual bool CanSave() const = 0;
+
+	/**
+	 * Checks whether this section can have its settings saved as default.
+	 *
+	 * @return true if saving as default is supported, false otherwise.
+	 */
+	virtual bool CanSaveDefaults() const = 0;
+
+	/**
+	 * Exports the settings in this section to the specified file.
+	 *
+	 * @param Filename The path to the file.
+	 * @return true if the settings were exported, false otherwise.
+	 */
+	virtual bool Export( const FString& Filename ) = 0;
+
+	/**
+	 * Gets the settings category that this section belongs to.
+	 *
+	 * @return The owner category.
+	 */
+	virtual TWeakPtr<ISettingsCategory> GetCategory() = 0;
+
+	/**
+	 * Gets the custom settings widget for this settings section.
+	 *
+	 * @return The custom widget.
+	 */
+	virtual TWeakPtr<SWidget> GetCustomWidget() const = 0;
+
+	/**
+	 * Gets the section's localized description text.
+	 *
+	 * @return Description text.
+	 */
+	virtual const FText& GetDescription() const = 0;
+
+	/**
+	 * Gets the section's localized display name.
+	 *
+	 * @return Display name.
+	 */
+	virtual const FText& GetDisplayName() const = 0;
+
+	/**
+	 * Gets the section's name.
+	 *
+	 * @return Section name.
+	 */
+	virtual const FName& GetName() const = 0;
+
+	/**
+	 * Gets the UObject holding the section's settings.
+	 *
+	 * @return Settings object.
+	 */
+	virtual TWeakObjectPtr<UObject> GetSettingsObject() const = 0;
+
+	/**
+	 * Gets the section's optional status text.
+	 *
+	 * @return Status text.
+	 */
+	virtual FText GetStatus() const = 0;
+
+	/**
+	 * Checks whether this section holds a settings object that saves directly to default configuration files.
+	 *
+	 * @return true if it has a default settings object, false otherwise.
+	 */
+	virtual bool HasDefaultSettingsObject() = 0;
+
+	/**
+	 * Imports the settings in this section from the specified file.
+	 *
+	 * @param Filename The path to the file.
+	 * @return true if the settings were imported, false otherwise.
+	 */
+	virtual bool Import( const FString& Filename ) = 0;
+
+	/**
+	 * Resets the settings in this section to their default value.
+	 *
+	 * @return true if the section was reset, false otherwise.
+	 */
+	virtual bool ResetDefaults() = 0;
+
+	/**
+	 * Saves the settings in this section.
+	 *
+	 * @return true if settings were saved, false otherwise.
+	 */
+	virtual bool Save() = 0;
+
+	/**
+	 * Saves the settings in this section as defaults.
+	 *
+	 * @return true if settings were saved, false otherwise.
+	 */
+	virtual bool SaveDefaults() = 0;
+
+public:
+
+	/**
+	 * A delegate that is executed to check whether a settings section can be edited.
+	 *
+	 * The return value indicates whether the section can be edited.
+	 */
+	DECLARE_DELEGATE_RetVal(bool, FOnCanEdit)
+	virtual FOnCanEdit& OnCanEdit() = 0;
+
+	/**
+	 * A delegate that is executed when a settings section should export its values to a file.
+	 *
+	 * The first parameter is the path to the file to export to.
+	 * The return value indicates whether exporting succeeded.
+	 */
+	DECLARE_DELEGATE_RetVal_OneParam(bool, FOnExport, const FString&)
+	virtual FOnExport& OnExport() = 0;
+
+	/**
+	 * A delegate that is executed when a settings section should import its values from a file.
+	 *
+	 * The first parameter is the path to the file to import from.
+	 * The return value indicates whether importing succeeded.
+	 */
+	DECLARE_DELEGATE_RetVal_OneParam(bool, FOnImport, const FString&)
+	virtual FOnImport& OnImport() = 0;
+
+	/**
+	 * A delegate that is executed when a settings section has been modified.
+	 *
+	 * The return value indicates whether the modifications should be saved.
+	 */
+	DECLARE_DELEGATE_RetVal(bool, FOnModified)
+	virtual FOnModified& OnModified() = 0;
+
+	/**
+	 * A delegate that is executed when a settings section should have its values reset to default.
+	 *
+	 * The return value indicates whether resetting to defaults succeeded.
+	 */
+	DECLARE_DELEGATE_RetVal(bool, FOnResetDefaults)
+	virtual FOnResetDefaults& OnResetDefaults() = 0;
+
+	/**
+	 * A delegate that is executed when a settings section should have its values saved.
+	 *
+	 * The return value indicates whether saving succeeded.
+	 */
+	DECLARE_DELEGATE_RetVal(bool, FOnSave)
+	virtual FOnSave& OnSave() = 0;
+
+	/**
+	 * A delegate that is executed when a settings section should have its values saved as default.
+	 *
+	 * The return value indicates whether saving as default succeeded.
+	 */
+	DECLARE_DELEGATE_RetVal(bool, FOnSaveDefaults)
+	virtual FOnSaveDefaults& OnSaveDefaults() = 0;
+
+	/**
+	 * A delegate that is executed to retrieve a status message for a settings section.
+	 *
+	 * The return value is status message.
+	 */
+	DECLARE_DELEGATE_RetVal(FText, FOnStatus)
+	virtual FOnStatus& OnStatus() = 0;
+
+public:
+
+	/** Virtual destructor. */
+	virtual ~ISettingsSection() { }
+};
+
+
+/** Type definition for shared pointers to instances of ISettingsSection. */
+typedef TSharedPtr<ISettingsSection> ISettingsSectionPtr;
+
+/** Type definition for shared references to instances of ISettingsSection. */
+typedef TSharedRef<ISettingsSection> ISettingsSectionRef;
diff --git a/Engine/Source/Developer/Settings/Public/ISettingsViewer.h b/Engine/Source/Developer/Settings/Public/ISettingsViewer.h
new file mode 100644
--- /dev/null
+++ b/Engine/Source/Developer/Settings/Public/ISettingsViewer.h
@@ -0,0 +1,25 @@
+// Copyright 1998-2014 Epic Games, Inc. All Rights Reserved.
+
+#pragma once
+
+
+/**
+ * Interface for settings tabs.
+ */
+class ISettingsViewer
+{
+public:
+
+	/**
+	 * Shows the settings tab that belongs to this viewer with the specified settings section pre-selected.
+	 *
+	 * @param CategoryName The name of the section's category.
+	 * @param SectionName The name of the section to select.
+	 */
+	virtual void ShowSettings( const FName& CategoryName, const FName& SectionName ) = 0;
+
+public:
+
+	/** Virtual destructor. */
+	virtual ~ISettingsViewer() { }
+};
diff --git a/Engine/Source/Developer/Settings/Public/Interfaces/ISettingsCategory.h b/Engine/Source/Developer/Settings/Public/Interfaces/ISettingsCategory.h
deleted file mode 100644
diff --git a/Engine/Source/Developer/Settings/Public/Interfaces/ISettingsContainer.h b/Engine/Source/Developer/Settings/Public/Interfaces/ISettingsContainer.h
deleted file mode 100644
diff --git a/Engine/Source/Developer/Settings/Public/Interfaces/ISettingsModule.h b/Engine/Source/Developer/Settings/Public/Interfaces/ISettingsModule.h
deleted file mode 100644
diff --git a/Engine/Source/Developer/Settings/Public/Interfaces/ISettingsSection.h b/Engine/Source/Developer/Settings/Public/Interfaces/ISettingsSection.h
deleted file mode 100644
diff --git a/Engine/Source/Developer/Settings/Public/Interfaces/ISettingsViewer.h b/Engine/Source/Developer/Settings/Public/Interfaces/ISettingsViewer.h
deleted file mode 100644
diff --git a/Engine/Source/Developer/Settings/Public/Settings.h b/Engine/Source/Developer/Settings/Public/Settings.h
deleted file mode 100644
diff --git a/Engine/Source/Developer/Settings/Settings.Build.cs b/Engine/Source/Developer/Settings/Settings.Build.cs
--- a/Engine/Source/Developer/Settings/Settings.Build.cs
+++ b/Engine/Source/Developer/Settings/Settings.Build.cs
@@ -9,14 +9,14 @@ public class Settings : ModuleRules
 		PublicDependencyModuleNames.AddRange(
 			new string[]
 			{
-				"CoreUObject",
+				"Core",
 			}
 		);
 
         PrivateDependencyModuleNames.AddRange(
             new string[]
             {
-                "Core",
+				"CoreUObject",
             }
         );
 
@@ -24,7 +24,6 @@ public class Settings : ModuleRules
 			new string[]
             {
                 "Developer/Settings/Private",
-                "Developer/Settings/Private/Widgets",
             }
 		);
 	}
diff --git a/Engine/Source/Developer/SettingsEditor/Private/SettingsEditorPrivatePCH.h b/Engine/Source/Developer/SettingsEditor/Private/SettingsEditorPrivatePCH.h
--- a/Engine/Source/Developer/SettingsEditor/Private/SettingsEditorPrivatePCH.h
+++ b/Engine/Source/Developer/SettingsEditor/Private/SettingsEditorPrivatePCH.h
@@ -12,6 +12,7 @@
 #include "DesktopPlatformModule.h"
 #include "ModuleManager.h"
 #include "PropertyEditing.h"
+#include "ISettingsCategory.h"
 #include "ISourceControlModule.h"
 
 
diff --git a/Engine/Source/Developer/SettingsEditor/Public/SettingsEditor.h b/Engine/Source/Developer/SettingsEditor/Public/SettingsEditor.h
--- a/Engine/Source/Developer/SettingsEditor/Public/SettingsEditor.h
+++ b/Engine/Source/Developer/SettingsEditor/Public/SettingsEditor.h
@@ -6,8 +6,9 @@
 /* Dependencies
  *****************************************************************************/
 
+#include "ISettingsContainer.h"
+#include "ISettingsSection.h"
 #include "ModuleInterface.h"
-#include "Settings.h"
 #include "SlateBasics.h"
 #include "EditorStyle.h"
 
diff --git a/Engine/Source/Developer/SourceCodeAccess/Private/SourceCodeAccessModule.cpp b/Engine/Source/Developer/SourceCodeAccess/Private/SourceCodeAccessModule.cpp
--- a/Engine/Source/Developer/SourceCodeAccess/Private/SourceCodeAccessModule.cpp
+++ b/Engine/Source/Developer/SourceCodeAccess/Private/SourceCodeAccessModule.cpp
@@ -2,9 +2,11 @@
 
 #include "SourceCodeAccessPrivatePCH.h"
 #include "Runtime/Core/Public/Features/IModularFeatures.h"
+#include "ISettingsModule.h"
+#include "ModuleManager.h"
 #include "SourceCodeAccessModule.h"
 #include "SourceCodeAccessSettings.h"
-#include "Settings.h"
+
 
 IMPLEMENT_MODULE( FSourceCodeAccessModule, SourceCodeAccess );
 
@@ -12,10 +14,11 @@ IMPLEMENT_MODULE( FSourceCodeAccessModule, SourceCodeAccess );
 
 static FName SourceCodeAccessorFeatureName(TEXT("SourceCodeAccessor"));
 
+
 FSourceCodeAccessModule::FSourceCodeAccessModule()
 	: CurrentSourceCodeAccessor(nullptr)
-{
-}
+{ }
+
 
 void FSourceCodeAccessModule::StartupModule()
 {
@@ -29,7 +32,8 @@ void FSourceCodeAccessModule::StartupModule()
 	IModularFeatures::Get().RegisterModularFeature(SourceCodeAccessorFeatureName, &DefaultSourceCodeAccessor);
 
 	// Register to display our settings
-	ISettingsModule* SettingsModule = ISettingsModule::Get();
+	ISettingsModule* SettingsModule = FModuleManager::GetModulePtr<ISettingsModule>("Settings");
+
 	if (SettingsModule != nullptr)
 	{
 		SettingsModule->RegisterSettings("Editor", "General", "Source Code",
@@ -40,10 +44,12 @@ void FSourceCodeAccessModule::StartupModule()
 	}
 }
 
+
 void FSourceCodeAccessModule::ShutdownModule()
 {
 	// Unregister our settings
-	ISettingsModule* SettingsModule = ISettingsModule::Get();
+	ISettingsModule* SettingsModule = FModuleManager::GetModulePtr<ISettingsModule>("Settings");
+
 	if (SettingsModule != nullptr)
 	{
 		SettingsModule->UnregisterSettings("Editor", "General", "Source Code");
@@ -57,16 +63,19 @@ void FSourceCodeAccessModule::ShutdownModule()
 	IModularFeatures::Get().OnModularFeatureUnregistered().RemoveAll(this);
 }
 
+
 bool FSourceCodeAccessModule::CanAccessSourceCode() const
 {
 	return CurrentSourceCodeAccessor->CanAccessSourceCode();
 }
 
+
 ISourceCodeAccessor& FSourceCodeAccessModule::GetAccessor() const
 {
 	return *CurrentSourceCodeAccessor;
 }
 
+
 void FSourceCodeAccessModule::SetAccessor(const FName& InName)
 {
 	const int32 FeatureCount = IModularFeatures::Get().GetModularFeatureImplementationCount(SourceCodeAccessorFeatureName);
@@ -89,6 +98,7 @@ void FSourceCodeAccessModule::SetAccessor(const FName& InName)
 	}
 }
 
+
 FLaunchingCodeAccessor& FSourceCodeAccessModule::OnLaunchingCodeAccessor()
 {
 	return LaunchingCodeAccessorDelegate;
@@ -99,11 +109,13 @@ FDoneLaunchingCodeAccessor& FSourceCodeAccessModule::OnDoneLaunchingCodeAccessor
 	return DoneLaunchingCodeAccessorDelegate;
 }
 
+
 FOpenFileFailed& FSourceCodeAccessModule::OnOpenFileFailed()
 {
 	return OpenFileFailedDelegate;
 }
 
+
 void FSourceCodeAccessModule::HandleModularFeatureRegistered(const FName& Type, IModularFeature* ModularFeature)
 {
 	if(Type == SourceCodeAccessorFeatureName)
@@ -133,6 +145,7 @@ void FSourceCodeAccessModule::HandleModularFeatureRegistered(const FName& Type,
 	}
 }
 
+
 void FSourceCodeAccessModule::HandleModularFeatureUnregistered(const FName& Type, IModularFeature* ModularFeature)
 {
 	if(Type == SourceCodeAccessorFeatureName && CurrentSourceCodeAccessor == static_cast<ISourceCodeAccessor*>(ModularFeature))
@@ -141,4 +154,5 @@ void FSourceCodeAccessModule::HandleModularFeatureUnregistered(const FName& Type
 	}
 }
 
-#undef LOCTEXT_NAMESPACE 
\ No newline at end of file
+
+#undef LOCTEXT_NAMESPACE
diff --git a/Engine/Source/Developer/SourceCodeAccess/Private/SourceCodeAccessModule.h b/Engine/Source/Developer/SourceCodeAccess/Private/SourceCodeAccessModule.h
--- a/Engine/Source/Developer/SourceCodeAccess/Private/SourceCodeAccessModule.h
+++ b/Engine/Source/Developer/SourceCodeAccess/Private/SourceCodeAccessModule.h
@@ -4,20 +4,31 @@
 
 #include "ISourceCodeAccessModule.h"
 #include "DefaultSourceCodeAccessor.h"
+#include "ModuleInterface.h"
+
 
 /**
- * Implementation of ISourceCodeAccessModule
+ * Implementation of ISourceCodeAccessModule.
  */
-class FSourceCodeAccessModule : public ISourceCodeAccessModule
+class FSourceCodeAccessModule
+	: public ISourceCodeAccessModule
 {
 public:
+
+	/** Default constructor. */
 	FSourceCodeAccessModule();
 
-	/** IModuleInterface implementation */
+public:
+
+	// IModuleInterface interface
+
 	virtual void StartupModule() override;
 	virtual void ShutdownModule() override;
 
-	/** ISourceCodeAccessModule implementation */
+public:
+
+	// ISourceCodeAccessModule interface
+
 	virtual bool CanAccessSourceCode() const override;
 	virtual ISourceCodeAccessor& GetAccessor() const override;
 	virtual void SetAccessor(const FName& InName) override;
@@ -26,6 +37,7 @@ public:
 	virtual FOpenFileFailed& OnOpenFileFailed() override;
 
 private:
+
 	/** Handle when one of the modular features we are interested in is registered */
 	void HandleModularFeatureRegistered(const FName& Type, IModularFeature* ModularFeature);
 
@@ -33,18 +45,19 @@ private:
 	void HandleModularFeatureUnregistered(const FName& Type, IModularFeature* ModularFeature);
 
 private:
-	/** Event delegate fired when launching code accessor */
+
+	/** Event delegate fired when launching code accessor. */
 	FLaunchingCodeAccessor LaunchingCodeAccessorDelegate;
 
-	/** Event delegate fired when done launching code accessor  */
+	/** Event delegate fired when done launching code accessor.  */
 	FDoneLaunchingCodeAccessor DoneLaunchingCodeAccessorDelegate;
 
-	/** Event delegate fired when opening a file has failed */
+	/** Event delegate fired when opening a file has failed. */
 	FOpenFileFailed OpenFileFailedDelegate;
 
-	/** The default accessor we will use if we have no IDE */
+	/** The default accessor we will use if we have no IDE. */
 	FDefaultSourceCodeAccessor DefaultSourceCodeAccessor;
 
-	/** The current accessor */
+	/** The current accessor. */
 	ISourceCodeAccessor* CurrentSourceCodeAccessor;
-};
\ No newline at end of file
+};
diff --git a/Engine/Source/Developer/SourceCodeAccess/Private/SourceCodeAccessPrivatePCH.h b/Engine/Source/Developer/SourceCodeAccess/Private/SourceCodeAccessPrivatePCH.h
--- a/Engine/Source/Developer/SourceCodeAccess/Private/SourceCodeAccessPrivatePCH.h
+++ b/Engine/Source/Developer/SourceCodeAccess/Private/SourceCodeAccessPrivatePCH.h
@@ -5,4 +5,3 @@
 #include "Core.h"
 #include "CoreUObject.h"
 #include "SlateBasics.h"
-#include "ModuleManager.h"
\ No newline at end of file
diff --git a/Engine/Source/Developer/TargetPlatform/Public/Interfaces/ITargetPlatformModule.h b/Engine/Source/Developer/TargetPlatform/Public/Interfaces/ITargetPlatformModule.h
--- a/Engine/Source/Developer/TargetPlatform/Public/Interfaces/ITargetPlatformModule.h
+++ b/Engine/Source/Developer/TargetPlatform/Public/Interfaces/ITargetPlatformModule.h
@@ -2,6 +2,8 @@
 
 #pragma once
 
+#include "ModuleInterface.h"
+
 
 /**
  * Interface for target platform modules.
@@ -12,12 +14,14 @@ class ITargetPlatformModule
 public:
 
 	/**
-	 * Gets the target platform for this module.
+	 * Gets the module's target platform.
+	 *
+	 * @return The target platform.
 	 */
 	virtual ITargetPlatform* GetTargetPlatform() = 0;
 
 public:
 
 	/** Virtual destructor. */
-	~ITargetPlatformModule() { }
+	virtual ~ITargetPlatformModule() { }
 };
diff --git a/Engine/Source/Developer/Windows/WindowsTargetPlatform/Private/WindowsTargetPlatformModule.cpp b/Engine/Source/Developer/Windows/WindowsTargetPlatform/Private/WindowsTargetPlatformModule.cpp
--- a/Engine/Source/Developer/Windows/WindowsTargetPlatform/Private/WindowsTargetPlatformModule.cpp
+++ b/Engine/Source/Developer/Windows/WindowsTargetPlatform/Private/WindowsTargetPlatformModule.cpp
@@ -1,14 +1,14 @@
 // Copyright 1998-2014 Epic Games, Inc. All Rights Reserved.
 
 #include "WindowsTargetPlatformPrivatePCH.h"
+#include "ISettingsModule.h"
+#include "ModuleManager.h"
 
 
 #define LOCTEXT_NAMESPACE "FWindowsTargetPlatformModule"
 
 
-/**
- * Holds the target platform singleton.
- */
+/** Holds the target platform singleton. */
 static ITargetPlatform* Singleton = nullptr;
 
 
@@ -71,7 +71,7 @@ public:
 		GConfig->GetArray(TEXT("/Script/WindowsTargetPlatform.WindowsTargetSettings"), TEXT("TargetedRHIs"), TargetSettings->TargetedRHIs, GEngineIni);
 		TargetSettings->AddToRoot();
 
-		ISettingsModule* SettingsModule = ISettingsModule::Get();
+		ISettingsModule* SettingsModule = FModuleManager::GetModulePtr<ISettingsModule>("Settings");
 
 		if (SettingsModule != nullptr)
 		{
@@ -90,7 +90,7 @@ public:
 		FCoreDelegates::GetHotfixDelegate(EHotfixDelegates::Test).Unbind();
 #endif
 
-		ISettingsModule* SettingsModule = ISettingsModule::Get();
+		ISettingsModule* SettingsModule = FModuleManager::GetModulePtr<ISettingsModule>("Settings");
 
 		if (SettingsModule != nullptr)
 		{
diff --git a/Engine/Source/Developer/Windows/WindowsTargetPlatform/Private/WindowsTargetPlatformPrivatePCH.h b/Engine/Source/Developer/Windows/WindowsTargetPlatform/Private/WindowsTargetPlatformPrivatePCH.h
--- a/Engine/Source/Developer/Windows/WindowsTargetPlatform/Private/WindowsTargetPlatformPrivatePCH.h
+++ b/Engine/Source/Developer/Windows/WindowsTargetPlatform/Private/WindowsTargetPlatformPrivatePCH.h
@@ -7,8 +7,6 @@
  *****************************************************************************/
 
 #include "Core.h"
-#include "ModuleManager.h"
-#include "Settings.h"
 #include "WindowsPlatformProperties.h"
 
 #if WITH_ENGINE
diff --git a/Engine/Source/Developer/iOS/IOSPlatformEditor/Private/IOSPlatformEditorModule.cpp b/Engine/Source/Developer/iOS/IOSPlatformEditor/Private/IOSPlatformEditorModule.cpp
--- a/Engine/Source/Developer/iOS/IOSPlatformEditor/Private/IOSPlatformEditorModule.cpp
+++ b/Engine/Source/Developer/iOS/IOSPlatformEditor/Private/IOSPlatformEditorModule.cpp
@@ -2,26 +2,35 @@
 
 #include "IOSPlatformEditorPrivatePCH.h"
 #include "IOSTargetSettingsCustomization.h"
-#include "Settings.h"
+#include "ISettingsModule.h"
+#include "ModuleInterface.h"
+#include "ModuleManager.h"
+
 
 #define LOCTEXT_NAMESPACE "FIOSPlatformEditorModule"
 
+
 /**
  * Module for iOS as a target platform
  */
-class FIOSPlatformEditorModule : public IModuleInterface
+class FIOSPlatformEditorModule
+	: public IModuleInterface
 {
 	// IModuleInterface interface
+
 	virtual void StartupModule() override
 	{
-		// Register the settings detail panel customization
+		// register settings detail panel customization
 		FPropertyEditorModule& PropertyModule = FModuleManager::LoadModuleChecked<FPropertyEditorModule>("PropertyEditor");
 		PropertyModule.RegisterCustomClassLayout(
 			"IOSRuntimeSettings",
-			FOnGetDetailCustomizationInstance::CreateStatic(&FIOSTargetSettingsCustomization::MakeInstance));
+			FOnGetDetailCustomizationInstance::CreateStatic(&FIOSTargetSettingsCustomization::MakeInstance)
+		);
+
 		PropertyModule.NotifyCustomizationModuleChanged();
 
-		ISettingsModule* SettingsModule = ISettingsModule::Get();
+		// register settings
+		ISettingsModule* SettingsModule = FModuleManager::GetModulePtr<ISettingsModule>("Settings");
 
 		if (SettingsModule != nullptr)
 		{
@@ -35,16 +44,16 @@ class FIOSPlatformEditorModule : public IModuleInterface
 
 	virtual void ShutdownModule() override
 	{
-		ISettingsModule* SettingsModule = ISettingsModule::Get();
+		ISettingsModule* SettingsModule = FModuleManager::GetModulePtr<ISettingsModule>("Settings");
 
 		if (SettingsModule != nullptr)
 		{
 			SettingsModule->UnregisterSettings("Project", "Platforms", "iOS");
 		}
 	}
-	// End of IModuleInterface interface
 };
 
+
 IMPLEMENT_MODULE(FIOSPlatformEditorModule, IOSPlatformEditor);
 
 #undef LOCTEXT_NAMESPACE
diff --git a/Engine/Source/Developer/iOS/IOSPlatformEditor/Private/IOSPlatformEditorPrivatePCH.h b/Engine/Source/Developer/iOS/IOSPlatformEditor/Private/IOSPlatformEditorPrivatePCH.h
--- a/Engine/Source/Developer/iOS/IOSPlatformEditor/Private/IOSPlatformEditorPrivatePCH.h
+++ b/Engine/Source/Developer/iOS/IOSPlatformEditor/Private/IOSPlatformEditorPrivatePCH.h
@@ -2,7 +2,9 @@
 
 #pragma once
 
+
+/* Private dependencies
+ *****************************************************************************/
+
 #include "Core.h"
-#include "ModuleInterface.h"
-#include "ModuleManager.h"
 #include "Engine.h"
diff --git a/Engine/Source/Editor/EditorLiveStreaming/EditorLiveStreaming.cpp b/Engine/Source/Editor/EditorLiveStreaming/EditorLiveStreaming.cpp
--- a/Engine/Source/Editor/EditorLiveStreaming/EditorLiveStreaming.cpp
+++ b/Engine/Source/Editor/EditorLiveStreaming/EditorLiveStreaming.cpp
@@ -5,16 +5,19 @@
 #include "Runtime/Core/Public/Features/IModularFeatures.h"
 #include "Runtime/Engine/Public/Features/ILiveStreamingService.h"
 #include "Runtime/SlateCore/Public/Rendering/SlateRenderer.h"
-#include "Developer/Settings/Public/Settings.h"
 #include "Editor/MainFrame/Public/Interfaces/IMainFrameModule.h"
 #include "Runtime/GameLiveStreaming/Public/IGameLiveStreaming.h"
+#include "ISettingsModule.h"
 #include "SNotificationList.h"
+#include "ModuleManager.h"
 #include "NotificationManager.h"
 
+
 IMPLEMENT_MODULE( FEditorLiveStreaming, EditorLiveStreaming );
 
 #define LOCTEXT_NAMESPACE "EditorLiveStreaming"
 
+
 FEditorLiveStreaming::FEditorLiveStreaming()
 	: bIsBroadcasting( false ),
 	  LiveStreamer( nullptr ),
@@ -28,7 +31,8 @@ FEditorLiveStreaming::FEditorLiveStreaming()
 void FEditorLiveStreaming::StartupModule()
 {
 	// Register settings
-	ISettingsModule* SettingsModule = ISettingsModule::Get();
+	ISettingsModule* SettingsModule = FModuleManager::GetModulePtr<ISettingsModule>("Settings");
+
 	if( ensure( SettingsModule != nullptr ) )
 	{
 		SettingsModule->RegisterSettings( "Editor", "General", "EditorLiveStreaming",
@@ -43,7 +47,8 @@ void FEditorLiveStreaming::StartupModule()
 void FEditorLiveStreaming::ShutdownModule()
 {
 	// Unregister settings
-	ISettingsModule* SettingsModule = ISettingsModule::Get();
+	ISettingsModule* SettingsModule = FModuleManager::GetModulePtr<ISettingsModule>("Settings");
+
 	if( SettingsModule != nullptr )
 	{
 		SettingsModule->UnregisterSettings( "Editor", "General", "EditorLiveStreaming" );
diff --git a/Engine/Source/Editor/EditorSettingsViewer/Private/EditorSettingsViewerModule.cpp b/Engine/Source/Editor/EditorSettingsViewer/Private/EditorSettingsViewerModule.cpp
--- a/Engine/Source/Editor/EditorSettingsViewer/Private/EditorSettingsViewerModule.cpp
+++ b/Engine/Source/Editor/EditorSettingsViewer/Private/EditorSettingsViewerModule.cpp
@@ -1,9 +1,16 @@
 // Copyright 1998-2014 Epic Games, Inc. All Rights Reserved.
 
 #include "EditorSettingsViewerPrivatePCH.h"
+#include "ISettingsCategory.h"
+#include "ISettingsModule.h"
+#include "ISettingsSection.h"
+#include "ISettingsViewer.h"
+#include "ModuleManager.h"
+#include "SDockTab.h"
+
 #include "Tests/AutomationTestSettings.h"
 #include "BlueprintEditorSettings.h"
-#include "SDockTab.h"
+
 
 #define LOCTEXT_NAMESPACE "FEditorSettingsViewerModule"
 
@@ -41,9 +48,9 @@ public:
 
 	// IModuleInterface interface
 
-	virtual void StartupModule( ) override
+	virtual void StartupModule() override
 	{
-		ISettingsModule* SettingsModule = ISettingsModule::Get();
+		ISettingsModule* SettingsModule = FModuleManager::GetModulePtr<ISettingsModule>("Settings");
 
 		if (SettingsModule != nullptr)
 		{
@@ -59,13 +66,13 @@ public:
 			.SetMenuType(ETabSpawnerMenuType::Hide);
 	}
 
-	virtual void ShutdownModule( ) override
+	virtual void ShutdownModule() override
 	{
 		FGlobalTabmanager::Get()->UnregisterNomadTabSpawner(EditorSettingsTabName);
 		UnregisterSettings();
 	}
 
-	virtual bool SupportsDynamicReloading( ) override
+	virtual bool SupportsDynamicReloading() override
 	{
 		return true;
 	}
@@ -74,44 +81,49 @@ protected:
 
 	/**
 	 * Registers general Editor settings.
+	 *
+	 * @param SettingsModule A reference to the settings module.
 	 */
 	void RegisterGeneralSettings( ISettingsModule& SettingsModule )
 	{
-		// automation settings
+		// automation
 		SettingsModule.RegisterSettings("Editor", "General", "AutomationTest",
 			LOCTEXT("AutomationSettingsName", "Automation"),
 			LOCTEXT("AutomationSettingsDescription", "Set up automation test assets."),
 			GetMutableDefault<UAutomationTestSettings>()
 		);
 
-		// input bindings
-		FSettingsSectionDelegates InputBindingDelegates;
-		InputBindingDelegates.ExportDelegate = FOnSettingsSectionExport::CreateRaw(this, &FEditorSettingsViewerModule::HandleInputBindingsExport);
-		InputBindingDelegates.ImportDelegate = FOnSettingsSectionImport::CreateRaw(this, &FEditorSettingsViewerModule::HandleInputBindingsImport);
-		InputBindingDelegates.ResetDefaultsDelegate = FOnSettingsSectionResetDefaults::CreateRaw(this, &FEditorSettingsViewerModule::HandleInputBindingsResetToDefault);
-		InputBindingDelegates.SaveDelegate = FOnSettingsSectionSave::CreateRaw(this, &FEditorSettingsViewerModule::HandleInputBindingsSave);
-
-		FSettingsSectionDelegates RegionAndLanguageDelegates;
-		RegionAndLanguageDelegates.ExportDelegate = FOnSettingsSectionExport::CreateRaw(this, &FEditorSettingsViewerModule::HandleRegionAndLanguageExport);
-		RegionAndLanguageDelegates.ImportDelegate = FOnSettingsSectionImport::CreateRaw(this, &FEditorSettingsViewerModule::HandleRegionAndLanguageImport);
-		RegionAndLanguageDelegates.SaveDefaultsDelegate = FOnSettingsSectionSaveDefaults::CreateRaw(this, &FEditorSettingsViewerModule::HandleRegionAndLanguageSaveDefaults);
-		RegionAndLanguageDelegates.ResetDefaultsDelegate = FOnSettingsSectionResetDefaults::CreateRaw(this, &FEditorSettingsViewerModule::HandleRegionAndLanguageResetToDefault);
-
-		SettingsModule.RegisterSettings("Editor", "General", "Internationalization",
+		// region & language
+		ISettingsSectionPtr RegionAndLanguageSettingsSection = SettingsModule.RegisterSettings("Editor", "General", "Internationalization",
 			LOCTEXT("InternationalizationSettingsModelName", "Region & Language"),
 			LOCTEXT("InternationalizationSettingsModelDescription", "Configure the editor's behavior to use a language and fit a region's culture."),
-			GetMutableDefault<UInternationalizationSettingsModel>(),
-			RegionAndLanguageDelegates
+			GetMutableDefault<UInternationalizationSettingsModel>()
 		);
 
+		if (RegionAndLanguageSettingsSection.IsValid())
+		{
+			RegionAndLanguageSettingsSection->OnExport().BindRaw(this, &FEditorSettingsViewerModule::HandleRegionAndLanguageExport);
+			RegionAndLanguageSettingsSection->OnImport().BindRaw(this, &FEditorSettingsViewerModule::HandleRegionAndLanguageImport);
+			RegionAndLanguageSettingsSection->OnSaveDefaults().BindRaw(this, &FEditorSettingsViewerModule::HandleRegionAndLanguageSaveDefaults);
+			RegionAndLanguageSettingsSection->OnResetDefaults().BindRaw(this, &FEditorSettingsViewerModule::HandleRegionAndLanguageResetToDefault);
+		}
+
+		// input bindings
 		TWeakPtr<SWidget> InputBindingEditorPanel = FModuleManager::LoadModuleChecked<IInputBindingEditorModule>("InputBindingEditor").CreateInputBindingEditorPanel();
-		SettingsModule.RegisterSettings("Editor", "General", "InputBindings",
+		ISettingsSectionPtr InputBindingSettingsSection = SettingsModule.RegisterSettings("Editor", "General", "InputBindings",
 			LOCTEXT("InputBindingsSettingsName", "Keyboard Shortcuts"),
 			LOCTEXT("InputBindingsSettingsDescription", "Configure keyboard shortcuts to quickly invoke operations."),
-			InputBindingEditorPanel.Pin().ToSharedRef(),
-			InputBindingDelegates
+			InputBindingEditorPanel.Pin().ToSharedRef()
 		);
 
+		if (InputBindingSettingsSection.IsValid())
+		{
+			InputBindingSettingsSection->OnExport().BindRaw(this, &FEditorSettingsViewerModule::HandleInputBindingsExport);
+			InputBindingSettingsSection->OnImport().BindRaw(this, &FEditorSettingsViewerModule::HandleInputBindingsImport);
+			InputBindingSettingsSection->OnResetDefaults().BindRaw(this, &FEditorSettingsViewerModule::HandleInputBindingsResetToDefault);
+			InputBindingSettingsSection->OnSave().BindRaw(this, &FEditorSettingsViewerModule::HandleInputBindingsSave);
+		}
+
 		// loading & saving features
 		SettingsModule.RegisterSettings("Editor", "General", "LoadingSaving",
 			LOCTEXT("LoadingSavingSettingsName", "Loading & Saving"),
@@ -138,7 +150,9 @@ protected:
 	}
 
 	/**
-	 * Registers Level Editor settings
+	 * Registers Level Editor settings.
+	 *
+	 * @param SettingsModule A reference to the settings module.
 	 */
 	void RegisterLevelEditorSettings( ISettingsModule& SettingsModule )
 	{
@@ -165,7 +179,9 @@ protected:
 	}
 
 	/**
-	 * Registers Other Tools settings
+	 * Registers Other Tools settings.
+	 *
+	 * @param SettingsModule A reference to the settings module.
 	 */
 	void RegisterContentEditorsSettings( ISettingsModule& SettingsModule )
 	{
@@ -198,12 +214,10 @@ protected:
 		);
 	}
 
-	/**
-	 * Unregisters all settings.
-	 */
-	void UnregisterSettings( )
+	/** Unregisters all settings. */
+	void UnregisterSettings()
 	{
-		ISettingsModule* SettingsModule = ISettingsModule::Get();
+		ISettingsModule* SettingsModule = FModuleManager::GetModulePtr<ISettingsModule>("Settings");
 
 		if (SettingsModule != nullptr)
 		{
@@ -230,10 +244,10 @@ protected:
 
 private:
 
-	// Handles creating the editor settings tab.
+	/** Handles creating the editor settings tab. */
 	TSharedRef<SDockTab> HandleSpawnSettingsTab( const FSpawnTabArgs& SpawnTabArgs )
 	{
-		ISettingsModule* SettingsModule = ISettingsModule::Get();
+		ISettingsModule* SettingsModule = FModuleManager::GetModulePtr<ISettingsModule>("Settings");
 		TSharedRef<SWidget> SettingsEditor = SNullWidget::NullWidget;
 
 		if (SettingsModule != nullptr)
@@ -379,7 +393,7 @@ private:
 
 private:
 
-	// Holds a pointer to the settings editor's view model.
+	/** Holds a pointer to the settings editor's view model. */
 	TWeakPtr<ISettingsEditorModel> SettingsEditorModelPtr;
 };
 
diff --git a/Engine/Source/Editor/EditorSettingsViewer/Private/EditorSettingsViewerPrivatePCH.h b/Engine/Source/Editor/EditorSettingsViewer/Private/EditorSettingsViewerPrivatePCH.h
--- a/Engine/Source/Editor/EditorSettingsViewer/Private/EditorSettingsViewerPrivatePCH.h
+++ b/Engine/Source/Editor/EditorSettingsViewer/Private/EditorSettingsViewerPrivatePCH.h
@@ -3,7 +3,7 @@
 #pragma once
 
 
-/* Dependencies
+/* Private dependencies
  *****************************************************************************/
 
 #include "UnrealEd.h"
@@ -11,5 +11,4 @@
 #include "InputBindingEditor.h"
 #include "InternationalizationSettingsModel.h"
 #include "MessageLog.h"
-#include "Settings.h"
 #include "SettingsEditor.h"
diff --git a/Engine/Source/Editor/EditorStyle/Private/EditorStylePrivatePCH.h b/Engine/Source/Editor/EditorStyle/Private/EditorStylePrivatePCH.h
--- a/Engine/Source/Editor/EditorStyle/Private/EditorStylePrivatePCH.h
+++ b/Engine/Source/Editor/EditorStyle/Private/EditorStylePrivatePCH.h
@@ -10,7 +10,6 @@
  *****************************************************************************/
 
 #include "ModuleManager.h"
-#include "Settings.h"
 #include "SlateBasics.h"
 #include "SlateStyle.h"
 
diff --git a/Engine/Source/Editor/EditorStyle/Private/SlateEditorStyle.h b/Engine/Source/Editor/EditorStyle/Private/SlateEditorStyle.h
--- a/Engine/Source/Editor/EditorStyle/Private/SlateEditorStyle.h
+++ b/Engine/Source/Editor/EditorStyle/Private/SlateEditorStyle.h
@@ -2,6 +2,9 @@
 
 #pragma once
 
+#include "ISettingsModule.h"
+#include "ModuleManager.h"
+
 
 /**
  * Declares the Editor's visual style.
@@ -17,7 +20,7 @@ public:
 
 #if WITH_EDITOR
 		Settings = GetMutableDefault<UEditorStyleSettings>();
-		ISettingsModule* SettingsModule = ISettingsModule::Get();
+		ISettingsModule* SettingsModule = FModuleManager::GetModulePtr<ISettingsModule>("Settings");
 
 		if (SettingsModule != nullptr)
 		{
@@ -35,7 +38,7 @@ public:
 
 	static void Shutdown()
 	{
-		ISettingsModule* SettingsModule = ISettingsModule::Get();
+		ISettingsModule* SettingsModule = FModuleManager::GetModulePtr<ISettingsModule>("Settings");
 
 		if (SettingsModule != nullptr)
 		{
diff --git a/Engine/Source/Editor/GameplayTagsEditor/Private/GameplayTagsEditorModule.cpp b/Engine/Source/Editor/GameplayTagsEditor/Private/GameplayTagsEditorModule.cpp
--- a/Engine/Source/Editor/GameplayTagsEditor/Private/GameplayTagsEditorModule.cpp
+++ b/Engine/Source/Editor/GameplayTagsEditor/Private/GameplayTagsEditorModule.cpp
@@ -5,53 +5,57 @@
 #include "GameplayTagContainerCustomization.h"
 #include "GameplayTagCustomization.h"
 #include "GameplayTagsSettings.h"
-#include "Settings.h"
+#include "ISettingsModule.h"
+#include "ModuleManager.h"
 
-#define LOCTEXT_NAMESPACE "GameplayTagEditor"
 
-class FGameplayTagsEditorModule : public IGameplayTagsEditorModule
-{
-	// Begin IModuleInterface
-	virtual void StartupModule() override;
-	virtual void ShutdownModule() override;
-	// End IModuleInterface
-};
+#define LOCTEXT_NAMESPACE "GameplayTagEditor"
 
-IMPLEMENT_MODULE( FGameplayTagsEditorModule, GameplayTagsEditor )
 
-void FGameplayTagsEditorModule::StartupModule()
+class FGameplayTagsEditorModule
+	: public IGameplayTagsEditorModule
 {
-	// Register the details customizer
-	FPropertyEditorModule& PropertyModule = FModuleManager::LoadModuleChecked<FPropertyEditorModule>("PropertyEditor");
-	PropertyModule.RegisterCustomPropertyTypeLayout("GameplayTagContainer", FOnGetPropertyTypeCustomizationInstance::CreateStatic(&FGameplayTagContainerCustomization::MakeInstance));
-	PropertyModule.RegisterCustomPropertyTypeLayout("GameplayTag", FOnGetPropertyTypeCustomizationInstance::CreateStatic(&FGameplayTagCustomization::MakeInstance));
+public:
 
-	TSharedPtr<FGameplayTagsGraphPanelPinFactory> GameplayTagsGraphPanelPinFactory = MakeShareable( new FGameplayTagsGraphPanelPinFactory() );
-	FEdGraphUtilities::RegisterVisualPinFactory(GameplayTagsGraphPanelPinFactory);
+	// IModuleInterface
 
-	if (UGameplayTagsManager::ShouldImportTagsFromINI())
+	virtual void StartupModule() override
 	{
-		if (ISettingsModule* SettingsModule = ISettingsModule::Get())
+		// Register the details customizer
+		FPropertyEditorModule& PropertyModule = FModuleManager::LoadModuleChecked<FPropertyEditorModule>("PropertyEditor");
+		PropertyModule.RegisterCustomPropertyTypeLayout("GameplayTagContainer", FOnGetPropertyTypeCustomizationInstance::CreateStatic(&FGameplayTagContainerCustomization::MakeInstance));
+		PropertyModule.RegisterCustomPropertyTypeLayout("GameplayTag", FOnGetPropertyTypeCustomizationInstance::CreateStatic(&FGameplayTagCustomization::MakeInstance));
+
+		TSharedPtr<FGameplayTagsGraphPanelPinFactory> GameplayTagsGraphPanelPinFactory = MakeShareable( new FGameplayTagsGraphPanelPinFactory() );
+		FEdGraphUtilities::RegisterVisualPinFactory(GameplayTagsGraphPanelPinFactory);
+
+		if (UGameplayTagsManager::ShouldImportTagsFromINI())
 		{
-			SettingsModule->RegisterSettings("Project", "Project", "GameplayTags",
-				LOCTEXT("GameplayTagSettingsName", "GameplayTags"),
-				LOCTEXT("GameplayTagSettingsNameDesc", "GameplayTag Settings"),
-				GetMutableDefault<UGameplayTagsSettings>()
-				);
+			if (ISettingsModule* SettingsModule = FModuleManager::GetModulePtr<ISettingsModule>("Settings"))
+			{
+				SettingsModule->RegisterSettings("Project", "Project", "GameplayTags",
+					LOCTEXT("GameplayTagSettingsName", "GameplayTags"),
+					LOCTEXT("GameplayTagSettingsNameDesc", "GameplayTag Settings"),
+					GetMutableDefault<UGameplayTagsSettings>()
+					);
+			}
 		}
 	}
-}
 
-void FGameplayTagsEditorModule::ShutdownModule()
-{
-	// This function may be called during shutdown to clean up your module.  For modules that support dynamic reloading,
-	// we call this function before unloading the module.
+	virtual void ShutdownModule() override
+	{
+		// This function may be called during shutdown to clean up your module.  For modules that support dynamic reloading,
+		// we call this function before unloading the module.
 	
 	
-	if (ISettingsModule* SettingsModule = ISettingsModule::Get())
-	{
-		SettingsModule->UnregisterSettings("Project", "Project", "GameplayTags");
+		if (ISettingsModule* SettingsModule = FModuleManager::GetModulePtr<ISettingsModule>("Settings"))
+		{
+			SettingsModule->UnregisterSettings("Project", "Project", "GameplayTags");
+		}
 	}
-}
+};
+
+
+IMPLEMENT_MODULE(FGameplayTagsEditorModule, GameplayTagsEditor)
 
 #undef LOCTEXT_NAMESPACE
\ No newline at end of file
diff --git a/Engine/Source/Editor/GameplayTagsEditor/Public/GameplayTagsEditorModule.h b/Engine/Source/Editor/GameplayTagsEditor/Public/GameplayTagsEditorModule.h
--- a/Engine/Source/Editor/GameplayTagsEditor/Public/GameplayTagsEditorModule.h
+++ b/Engine/Source/Editor/GameplayTagsEditor/Public/GameplayTagsEditorModule.h
@@ -2,6 +2,7 @@
 
 #pragma once
 
+#include "ModuleInterface.h"
 #include "ModuleManager.h"
 
 
diff --git a/Engine/Source/Editor/HardwareTargeting/Private/HardwareTargetingModule.cpp b/Engine/Source/Editor/HardwareTargeting/Private/HardwareTargetingModule.cpp
--- a/Engine/Source/Editor/HardwareTargeting/Private/HardwareTargetingModule.cpp
+++ b/Engine/Source/Editor/HardwareTargeting/Private/HardwareTargetingModule.cpp
@@ -3,15 +3,18 @@
 #include "HardwareTargetingPrivatePCH.h"
 #include "HardwareTargetingModule.h"
 #include "HardwareTargetingSettings.h"
-#include "Settings.h"
 #include "Internationalization.h"
+#include "ISettingsModule.h"
+#include "ModuleManager.h"
 #include "SDecoratedEnumCombo.h"
 
 #include "Runtime/Engine/Classes/Engine/RendererSettings.h"
 #include "Editor/Documentation/Public/IDocumentation.h"
 
+
 #define LOCTEXT_NAMESPACE "HardwareTargeting"
 
+
 //////////////////////////////////////////////////////////////////////////
 // FMetaSettingGatherer
 
@@ -140,13 +143,16 @@ private:
 
 void FHardwareTargetingModule::StartupModule()
 {
-	ISettingsModule& SettingsModule = *ISettingsModule::Get();
+	ISettingsModule* SettingsModule = FModuleManager::GetModulePtr<ISettingsModule>("Settings");
 		
-	SettingsModule.RegisterSettings("Project", "Project", "HardwareTargeting",
-		LOCTEXT("HardwareTargetingSettingsName", "Target Hardware"),
-		LOCTEXT("HardwareTargetingSettingsDescription", "Options for choosing which class of hardware to target"),
-		GetMutableDefault<UHardwareTargetingSettings>()
-	);
+	if (SettingsModule != nullptr)
+	{
+		SettingsModule->RegisterSettings("Project", "Project", "HardwareTargeting",
+			LOCTEXT("HardwareTargetingSettingsName", "Target Hardware"),
+			LOCTEXT("HardwareTargetingSettingsDescription", "Options for choosing which class of hardware to target"),
+			GetMutableDefault<UHardwareTargetingSettings>()
+		);
+	}
 
 	// Apply any settings on startup if necessary
 	ApplyHardwareTargetingSettings();
diff --git a/Engine/Source/Editor/HardwareTargeting/Private/HardwareTargetingSettings.cpp b/Engine/Source/Editor/HardwareTargeting/Private/HardwareTargetingSettings.cpp
--- a/Engine/Source/Editor/HardwareTargeting/Private/HardwareTargetingSettings.cpp
+++ b/Engine/Source/Editor/HardwareTargeting/Private/HardwareTargetingSettings.cpp
@@ -3,18 +3,18 @@
 #include "HardwareTargetingPrivatePCH.h"
 #include "HardwareTargetingModule.h"
 #include "HardwareTargetingSettings.h"
-#include "Settings.h"
 #include "Internationalization.h"
 #include "SDecoratedEnumCombo.h"
 
+
 UHardwareTargetingSettings::UHardwareTargetingSettings(const FObjectInitializer& ObjectInitializer)
 	: Super(ObjectInitializer)
 	, TargetedHardwareClass(EHardwareClass::Unspecified)
 	, AppliedTargetedHardwareClass(EHardwareClass::Unspecified)
 	, DefaultGraphicsPerformance(EGraphicsPreset::Unspecified)
 	, AppliedDefaultGraphicsPerformance(EGraphicsPreset::Unspecified)
-{
-}
+{ }
+
 
 bool UHardwareTargetingSettings::HasPendingChanges() const
 {
@@ -26,7 +26,8 @@ bool UHardwareTargetingSettings::HasPendingChanges() const
 	return AppliedTargetedHardwareClass != TargetedHardwareClass || AppliedDefaultGraphicsPerformance != DefaultGraphicsPerformance;
 }
 
+
 void UHardwareTargetingSettings::PostEditChangeProperty( struct FPropertyChangedEvent& PropertyChangedEvent )
 {
 	SettingChangedEvent.Broadcast();
-}
\ No newline at end of file
+}
diff --git a/Engine/Source/Editor/IntroTutorials/Private/IntroTutorials.cpp b/Engine/Source/Editor/IntroTutorials/Private/IntroTutorials.cpp
--- a/Engine/Source/Editor/IntroTutorials/Private/IntroTutorials.cpp
+++ b/Engine/Source/Editor/IntroTutorials/Private/IntroTutorials.cpp
@@ -17,7 +17,7 @@
 #include "EditorTutorialSettings.h"
 #include "TutorialStateSettings.h"
 #include "TutorialSettings.h"
-#include "Settings.h"
+#include "ISettingsModule.h"
 #include "EditorTutorial.h"
 #include "SEditorTutorials.h"
 #include "STutorialsBrowser.h"
@@ -32,11 +32,14 @@
 #include "EngineBuildSettings.h"
 #include "IDocumentation.h"
 #include "SDockTab.h"
+#include "ModuleManager.h"
+
 
 #define LOCTEXT_NAMESPACE "IntroTutorials"
 
 IMPLEMENT_MODULE(FIntroTutorials, IntroTutorials)
 
+
 FIntroTutorials::FIntroTutorials()
 	: CurrentObjectClass(nullptr)
 	, ContentIntroCurve(nullptr)
@@ -117,7 +120,8 @@ void FIntroTutorials::StartupModule()
 	}
 
 	// Register to display our settings
-	ISettingsModule* SettingsModule = ISettingsModule::Get();
+	ISettingsModule* SettingsModule = FModuleManager::GetModulePtr<ISettingsModule>("Settings");
+
 	if (SettingsModule != nullptr)
 	{
 		SettingsModule->RegisterSettings("Editor", "General", "Tutorials",
@@ -173,7 +177,7 @@ void FIntroTutorials::ShutdownModule()
 		MainFrameModule.OnMainFrameCreationFinished().RemoveAll(this);
 	}
 
-	ISettingsModule* SettingsModule = ISettingsModule::Get();
+	ISettingsModule* SettingsModule = FModuleManager::GetModulePtr<ISettingsModule>("Settings");
 	if (SettingsModule != nullptr)
 	{
 		SettingsModule->UnregisterSettings("Editor", "General", "Tutorials");
diff --git a/Engine/Source/Editor/LevelEditor/Private/LevelEditorToolBar.cpp b/Engine/Source/Editor/LevelEditor/Private/LevelEditorToolBar.cpp
--- a/Engine/Source/Editor/LevelEditor/Private/LevelEditorToolBar.cpp
+++ b/Engine/Source/Editor/LevelEditor/Private/LevelEditorToolBar.cpp
@@ -16,13 +16,18 @@
 #include "EngineBuildSettings.h"
 #include "Matinee/MatineeActor.h"
 #include "Engine/LevelScriptBlueprint.h"
-#include "Settings.h"
+#include "ISettingsCategory.h"
+#include "ISettingsContainer.h"
+#include "ISettingsModule.h"
+#include "ISettingsSection.h"
 #include "DesktopPlatformModule.h"
 #include "Editor/ClassViewer/Public/ClassViewerModule.h"
 #include "Editor/ClassViewer/Public/ClassViewerFilter.h"
 #include "KismetEditorUtilities.h"
 #include "ISourceControlModule.h"
 #include "SVolumeControl.h"
+#include "ModuleManager.h"
+
 
 namespace LevelEditorActionHelpers
 {
@@ -783,7 +788,7 @@ void LevelEditorActionHelpers::OnSelectGameModeClassPicked(UClass* InChosenClass
 	{
 		UGameMapsSettings::SetGlobalDefaultGameMode(InChosenClass? InChosenClass->GetPathName() : FString());
 
-		ISettingsModule* SettingsModule = ISettingsModule::Get();
+		ISettingsModule* SettingsModule = FModuleManager::GetModulePtr<ISettingsModule>("Settings");
 		if (SettingsModule != nullptr)
 		{
 			ISettingsContainerPtr SettingsContainer = SettingsModule->GetContainer("Project");
diff --git a/Engine/Source/Editor/LevelEditor/Private/SLevelViewport.cpp b/Engine/Source/Editor/LevelEditor/Private/SLevelViewport.cpp
--- a/Engine/Source/Editor/LevelEditor/Private/SLevelViewport.cpp
+++ b/Engine/Source/Editor/LevelEditor/Private/SLevelViewport.cpp
@@ -25,7 +25,7 @@
 #include "LevelUtils.h"
 #include "HighresScreenshotUI.h"
 #include "SCaptureRegionWidget.h"
-#include "Settings.h"
+#include "ISettingsModule.h"
 #include "SceneOutlinerTreeItems.h"
 #include "BufferVisualizationData.h"
 #include "EditorViewportCommands.h"
diff --git a/Engine/Source/Editor/MainFrame/Private/MainFramePrivatePCH.h b/Engine/Source/Editor/MainFrame/Private/MainFramePrivatePCH.h
--- a/Engine/Source/Editor/MainFrame/Private/MainFramePrivatePCH.h
+++ b/Engine/Source/Editor/MainFrame/Private/MainFramePrivatePCH.h
@@ -29,7 +29,7 @@
 //#include "SearchUI.h"
 #include "SourceCodeNavigation.h"
 #include "SourceControlWindows.h"
-#include "Settings.h"
+#include "ISettingsModule.h"
 #include "TargetPlatform.h"
 #include "PlatformInfo.h"
 #include "UnrealEd.h"
diff --git a/Engine/Source/Editor/MainFrame/Private/Menus/SettingsMenu.h b/Engine/Source/Editor/MainFrame/Private/Menus/SettingsMenu.h
--- a/Engine/Source/Editor/MainFrame/Private/Menus/SettingsMenu.h
+++ b/Engine/Source/Editor/MainFrame/Private/Menus/SettingsMenu.h
@@ -2,6 +2,12 @@
 
 #pragma once
 
+#include "ISettingsCategory.h"
+#include "ISettingsContainer.h"
+#include "ISettingsModule.h"
+#include "ISettingsSection.h"
+#include "ModuleManager.h"
+
 
 #define LOCTEXT_NAMESPACE "FSettingsMenu"
 
@@ -21,7 +27,7 @@ public:
 	 */
 	static void MakeMenu( FMenuBuilder& MenuBuilder, FName SettingsContainerName )
 	{
-		ISettingsModule* SettingsModule = ISettingsModule::Get();
+		ISettingsModule* SettingsModule = FModuleManager::GetModulePtr<ISettingsModule>("Settings");
 
 		if (SettingsModule == nullptr)
 		{
diff --git a/Engine/Source/Editor/NewsFeed/Private/NewsFeedModule.cpp b/Engine/Source/Editor/NewsFeed/Private/NewsFeedModule.cpp
--- a/Engine/Source/Editor/NewsFeed/Private/NewsFeedModule.cpp
+++ b/Engine/Source/Editor/NewsFeed/Private/NewsFeedModule.cpp
@@ -1,6 +1,8 @@
 // Copyright 1998-2014 Epic Games, Inc. All Rights Reserved.
 
 #include "NewsFeedPrivatePCH.h"
+#include "ISettingsModule.h"
+#include "ModuleManager.h"
 
 
 #define LOCTEXT_NAMESPACE "FNewsFeedModule"
@@ -19,7 +21,7 @@ public:
 	virtual void StartupModule( ) override
 	{
 		// register settings
-		ISettingsModule* SettingsModule = ISettingsModule::Get();
+		ISettingsModule* SettingsModule = FModuleManager::GetModulePtr<ISettingsModule>("Settings");
 
 		if (SettingsModule != nullptr)
 		{
@@ -34,7 +36,7 @@ public:
 	virtual void ShutdownModule( ) override
 	{
 		// unregister settings
-		ISettingsModule* SettingsModule = ISettingsModule::Get();
+		ISettingsModule* SettingsModule = FModuleManager::GetModulePtr<ISettingsModule>("Settings");
 
 		if (SettingsModule != nullptr)
 		{
diff --git a/Engine/Source/Editor/NewsFeed/Private/NewsFeedPrivatePCH.h b/Engine/Source/Editor/NewsFeed/Private/NewsFeedPrivatePCH.h
--- a/Engine/Source/Editor/NewsFeed/Private/NewsFeedPrivatePCH.h
+++ b/Engine/Source/Editor/NewsFeed/Private/NewsFeedPrivatePCH.h
@@ -20,12 +20,10 @@
 #include "IHttpResponse.h"
 #include "ImageWrapper.h"
 #include "Json.h"
-#include "ModuleManager.h"
 #include "Online.h"
 #include "OnlineSubsystemTypes.h"
 #include "OnlineTitleFileInterface.h"
 #include "SecureHash.h"
-#include "Settings.h"
 #include "Ticker.h"
 
 
diff --git a/Engine/Source/Editor/NewsFeed/Private/Widgets/SNewsFeed.cpp b/Engine/Source/Editor/NewsFeed/Private/Widgets/SNewsFeed.cpp
--- a/Engine/Source/Editor/NewsFeed/Private/Widgets/SNewsFeed.cpp
+++ b/Engine/Source/Editor/NewsFeed/Private/Widgets/SNewsFeed.cpp
@@ -1,6 +1,7 @@
 // Copyright 1998-2014 Epic Games, Inc. All Rights Reserved.
 
 #include "NewsFeedPrivatePCH.h"
+#include "ISettingsModule.h"
 #include "SThrobber.h"
 #include "SHyperlink.h"
 
diff --git a/Engine/Source/Editor/ProjectSettingsViewer/Private/ProjectSettingsViewerModule.cpp b/Engine/Source/Editor/ProjectSettingsViewer/Private/ProjectSettingsViewerModule.cpp
--- a/Engine/Source/Editor/ProjectSettingsViewer/Private/ProjectSettingsViewerModule.cpp
+++ b/Engine/Source/Editor/ProjectSettingsViewer/Private/ProjectSettingsViewerModule.cpp
@@ -4,10 +4,16 @@
 #include "Engine/Console.h"
 #include "ProjectTargetPlatformEditor.h"
 #include "CookerSettings.h"
+#include "ISettingsCategory.h"
+#include "ISettingsModule.h"
+#include "ISettingsViewer.h"
+#include "ModuleManager.h"
+#include "SDockTab.h"
+
 #include "Runtime/Engine/Classes/Engine/UserInterfaceSettings.h"
 #include "Runtime/Engine/Classes/Engine/RendererSettings.h"
 #include "Runtime/Engine/Classes/Sound/AudioSettings.h"
-#include "SDockTab.h"
+
 
 #define LOCTEXT_NAMESPACE "FProjectSettingsViewerModule"
 
@@ -47,7 +53,7 @@ public:
 
 	virtual void StartupModule( ) override
 	{
-		ISettingsModule* SettingsModule = ISettingsModule::Get();
+		ISettingsModule* SettingsModule = FModuleManager::GetModulePtr<ISettingsModule>("Settings");
 
 		if (SettingsModule != nullptr)
 		{
@@ -77,6 +83,8 @@ protected:
 
 	/**
 	 * Registers Engine settings.
+	 *
+	 * @param SettingsModule A reference to the settings module.
 	 */
 	void RegisterEngineSettings( ISettingsModule& SettingsModule )
 	{
@@ -174,6 +182,8 @@ protected:
 
 	/**
 	 * Registers Project settings.
+	 *
+	 * @param SettingsModule A reference to the settings module.
 	 */
 	void RegisterProjectSettings( ISettingsModule& SettingsModule )
 	{
@@ -228,12 +238,10 @@ protected:
 		);*/
 	}
 
-	/**
-	 * Unregisters all settings.
-	 */
+	/** Unregisters all settings. */
 	void UnregisterSettings( )
 	{
-		ISettingsModule* SettingsModule = ISettingsModule::Get();
+		ISettingsModule* SettingsModule = FModuleManager::GetModulePtr<ISettingsModule>("Settings");
 
 		if (SettingsModule != nullptr)
 		{
@@ -262,11 +270,11 @@ protected:
 
 private:
 
-	// Handles creating the project settings tab.
+	/** Handles creating the project settings tab. */
 	TSharedRef<SDockTab> HandleSpawnSettingsTab( const FSpawnTabArgs& SpawnTabArgs )
 	{
-		ISettingsModule* SettingsModule = ISettingsModule::Get();
 		TSharedRef<SWidget> SettingsEditor = SNullWidget::NullWidget;
+		ISettingsModule* SettingsModule = FModuleManager::GetModulePtr<ISettingsModule>("Settings");
 
 		if (SettingsModule != nullptr)
 		{
@@ -291,7 +299,7 @@ private:
 
 private:
 
-	// Holds a pointer to the settings editor's view model.
+	/** Holds a pointer to the settings editor's view model. */
 	TWeakPtr<ISettingsEditorModel> SettingsEditorModelPtr;
 };
 
diff --git a/Engine/Source/Editor/ProjectSettingsViewer/Private/ProjectSettingsViewerPrivatePCH.h b/Engine/Source/Editor/ProjectSettingsViewer/Private/ProjectSettingsViewerPrivatePCH.h
--- a/Engine/Source/Editor/ProjectSettingsViewer/Private/ProjectSettingsViewerPrivatePCH.h
+++ b/Engine/Source/Editor/ProjectSettingsViewer/Private/ProjectSettingsViewerPrivatePCH.h
@@ -3,12 +3,12 @@
 #pragma  once
 
 
-/* Dependencies
+/* Private dependencies
  *****************************************************************************/
 
 #include "Engine.h"
 #include "EngineSettings.h"
 #include "ModuleManager.h"
+#include "MoviePlayerSettings.h"
 #include "SettingsEditor.h"
 #include "UnrealEd.h"
-#include "MoviePlayerSettings.h"
diff --git a/Engine/Source/Editor/ProjectTargetPlatformEditor/Private/ProjectTargetPlatformEditorModule.cpp b/Engine/Source/Editor/ProjectTargetPlatformEditor/Private/ProjectTargetPlatformEditorModule.cpp
--- a/Engine/Source/Editor/ProjectTargetPlatformEditor/Private/ProjectTargetPlatformEditorModule.cpp
+++ b/Engine/Source/Editor/ProjectTargetPlatformEditor/Private/ProjectTargetPlatformEditorModule.cpp
@@ -1,18 +1,21 @@
 // Copyright 1998-2014 Epic Games, Inc. All Rights Reserved.
 
 #include "ProjectTargetPlatformEditorPrivatePCH.h"
-#include "Settings.h"
+#include "ISettingsModule.h"
+
 
 #define LOCTEXT_NAMESPACE "FProjectTargetPlatformEditorModule"
 
+
 /**
  * Implements the platform target platform editor module
  */
-class FProjectTargetPlatformEditorModule : public IProjectTargetPlatformEditorModule
+class FProjectTargetPlatformEditorModule
+	: public IProjectTargetPlatformEditorModule
 {
 public:
 
-	// Begin IProjectTargetPlatformEditorModule interface
+	// IProjectTargetPlatformEditorModule interface
 
 	virtual TWeakPtr<SWidget> CreateProjectTargetPlatformEditorPanel() override
 	{
@@ -131,14 +134,13 @@ public:
 		return true;
 	}
 
-	// End IProjectTargetPlatformEditorModule interface
-
 private:
 
 	// Holds the collection of created editor panels.
 	TArray<TSharedPtr<SWidget> > EditorPanels;
 };
 
+
 #undef LOCTEXT_NAMESPACE
 
 IMPLEMENT_MODULE(FProjectTargetPlatformEditorModule, ProjectTargetPlatformEditor);
diff --git a/Engine/Source/Editor/TextureEditor/Private/TextureEditorModule.cpp b/Engine/Source/Editor/TextureEditor/Private/TextureEditorModule.cpp
--- a/Engine/Source/Editor/TextureEditor/Private/TextureEditorModule.cpp
+++ b/Engine/Source/Editor/TextureEditor/Private/TextureEditorModule.cpp
@@ -2,6 +2,7 @@
 
 
 #include "TextureEditorPrivatePCH.h"
+#include "ISettingsModule.h"
 #include "ModuleManager.h"
 
 
@@ -50,7 +51,7 @@ public:
 		ToolBarExtensibilityManager = MakeShareable(new FExtensibilityManager);
 
 		// register settings
-		ISettingsModule* SettingsModule = ISettingsModule::Get();
+		ISettingsModule* SettingsModule = FModuleManager::GetModulePtr<ISettingsModule>("Settings");
 
 		if (SettingsModule != nullptr)
 		{
@@ -65,7 +66,7 @@ public:
 	virtual void ShutdownModule( ) override
 	{
 		// unregister settings
-		ISettingsModule* SettingsModule = ISettingsModule::Get();
+		ISettingsModule* SettingsModule = FModuleManager::GetModulePtr<ISettingsModule>("Settings");
 
 		if (SettingsModule != nullptr)
 		{
diff --git a/Engine/Source/Editor/TextureEditor/Private/TextureEditorPrivatePCH.h b/Engine/Source/Editor/TextureEditor/Private/TextureEditorPrivatePCH.h
--- a/Engine/Source/Editor/TextureEditor/Private/TextureEditorPrivatePCH.h
+++ b/Engine/Source/Editor/TextureEditor/Private/TextureEditorPrivatePCH.h
@@ -15,7 +15,6 @@
 #include "SceneViewport.h"
 #include "SColorPicker.h"
 #include "ScopedTransaction.h"
-#include "Settings.h"
 #include "Texture2DPreview.h"
 #include "Toolkits/IToolkitHost.h"
 #include "Editor/WorkspaceMenuStructure/Public/WorkspaceMenuStructureModule.h"
diff --git a/Engine/Source/Editor/TextureEditor/Private/TextureEditorToolkit.cpp b/Engine/Source/Editor/TextureEditor/Private/TextureEditorToolkit.cpp
--- a/Engine/Source/Editor/TextureEditor/Private/TextureEditorToolkit.cpp
+++ b/Engine/Source/Editor/TextureEditor/Private/TextureEditorToolkit.cpp
@@ -2,6 +2,7 @@
 
 #include "TextureEditorPrivatePCH.h"
 #include "Factories.h"
+#include "ISettingsModule.h"
 #include "SDockTab.h"
 #include "SNumericEntryBox.h"
 
diff --git a/Engine/Source/Editor/UMGEditor/Private/Designer/SDesignerView.cpp b/Engine/Source/Editor/UMGEditor/Private/Designer/SDesignerView.cpp
--- a/Engine/Source/Editor/UMGEditor/Private/Designer/SDesignerView.cpp
+++ b/Engine/Source/Editor/UMGEditor/Private/Designer/SDesignerView.cpp
@@ -2,7 +2,7 @@
 
 #include "UMGEditorPrivatePCH.h"
 
-#include "Settings.h"
+#include "ISettingsModule.h"
 
 #include "Designer/DesignTimeUtils.h"
 
diff --git a/Engine/Source/Editor/UnrealEd/Private/Kismet2/DebuggerCommands.cpp b/Engine/Source/Editor/UnrealEd/Private/Kismet2/DebuggerCommands.cpp
--- a/Engine/Source/Editor/UnrealEd/Private/Kismet2/DebuggerCommands.cpp
+++ b/Engine/Source/Editor/UnrealEd/Private/Kismet2/DebuggerCommands.cpp
@@ -9,7 +9,7 @@
 #include "BlueprintUtilities.h"
 
 #include "LauncherServices.h"
-#include "Settings.h"
+#include "ISettingsModule.h"
 #include "TargetDeviceServices.h"
 #include "IMainFrameModule.h"
 
diff --git a/Engine/Source/Editor/UnrealEd/Private/PlayLevel.cpp b/Engine/Source/Editor/UnrealEd/Private/PlayLevel.cpp
--- a/Engine/Source/Editor/UnrealEd/Private/PlayLevel.cpp
+++ b/Engine/Source/Editor/UnrealEd/Private/PlayLevel.cpp
@@ -22,7 +22,7 @@
 #include "UObjectToken.h"
 #include "MapErrors.h"
 #include "LauncherServices.h"
-#include "Settings.h"
+#include "ISettingsModule.h"
 #include "TargetDeviceServices.h"
 #include "GameProjectGenerationModule.h"
 #include "SourceCodeNavigation.h"
diff --git a/Engine/Source/Runtime/Launch/Private/Mac/LaunchMac.cpp b/Engine/Source/Runtime/Launch/Private/Mac/LaunchMac.cpp
--- a/Engine/Source/Runtime/Launch/Private/Mac/LaunchMac.cpp
+++ b/Engine/Source/Runtime/Launch/Private/Mac/LaunchMac.cpp
@@ -6,7 +6,7 @@
 
 #if WITH_EDITOR
 	#include "MainFrame.h"
-	#include "Settings.h"
+	#include "ISettingsModule.h"
 #endif
 
 #include <signal.h>
