Settings: Made private headers compilable individually

diff --git a/Engine/Source/Developer/Settings/Private/SettingsCategory.h b/Engine/Source/Developer/Settings/Private/SettingsCategory.h
--- a/Engine/Source/Developer/Settings/Private/SettingsCategory.h
+++ b/Engine/Source/Developer/Settings/Private/SettingsCategory.h
@@ -3,6 +3,10 @@
 #pragma once
 
 
+// forward declarations
+class FSettingsSection;
+
+
 /**
  * Implements a settings category.
  */
diff --git a/Engine/Source/Developer/Settings/Private/SettingsContainer.h b/Engine/Source/Developer/Settings/Private/SettingsContainer.h
--- a/Engine/Source/Developer/Settings/Private/SettingsContainer.h
+++ b/Engine/Source/Developer/Settings/Private/SettingsContainer.h
@@ -3,6 +3,10 @@
 #pragma once
 
 
+// forward declarations
+class FSettingsCategory;
+
+
 /**
  * Implements a settings container.
  */
