Editor - The settings system now creates a category if one did not exist when you attempt to add a new section to a category that did not previously exist, instead of ignoring your request.

diff --git a/Engine/Source/Developer/Settings/Private/SettingsContainer.cpp b/Engine/Source/Developer/Settings/Private/SettingsContainer.cpp
--- a/Engine/Source/Developer/Settings/Private/SettingsContainer.cpp
+++ b/Engine/Source/Developer/Settings/Private/SettingsContainer.cpp
@@ -20,7 +20,8 @@ ISettingsSectionPtr FSettingsContainer::AddSection( const FName& CategoryName, c
 
 	if (!Category.IsValid())
 	{
-		return nullptr;
+		DescribeCategory(CategoryName, FText::FromString(FName::NameToDisplayString(CategoryName.ToString(), false)), FText::GetEmpty());
+		Category = Categories.FindRef(CategoryName);
 	}
 
 	ISettingsSectionRef Section = Category->AddSection(SectionName, InDisplayName, InDescription, SettingsObject);
@@ -36,7 +37,8 @@ ISettingsSectionPtr FSettingsContainer::AddSection( const FName& CategoryName, c
 
 	if (!Category.IsValid())
 	{
-		return nullptr;
+		DescribeCategory(CategoryName, FText::FromString(FName::NameToDisplayString(CategoryName.ToString(), false)), FText::GetEmpty());
+		Category = Categories.FindRef(CategoryName);
 	}
 
 	ISettingsSectionRef Section = Category->AddSection(SectionName, InDisplayName, InDescription, CustomWidget);
