Added a priority boost for a setting category named the same as FApp::GetProjectName() so game-specific settings sections can get sorted to the top if desired (by overriding GetCategoryName() to return FApp::GetProjectName())
#rnx

diff --git a/Engine/Source/Developer/Settings/Private/SettingsContainer.cpp b/Engine/Source/Developer/Settings/Private/SettingsContainer.cpp
--- a/Engine/Source/Developer/Settings/Private/SettingsContainer.cpp
+++ b/Engine/Source/Developer/Settings/Private/SettingsContainer.cpp
@@ -2,7 +2,7 @@
 
 #include "SettingsContainer.h"
 #include "UObject/WeakObjectPtr.h"
-
+#include "Misc/App.h"
 
 /* FSettingsContainer structors
  *****************************************************************************/
@@ -96,7 +96,9 @@ int32 FSettingsContainer::GetCategories( TArray<ISettingsCategoryPtr>& OutCatego
 	OutCategories.Empty(Categories.Num());
 
 	static const FName AdvancedCategoryName("Advanced");
+	const FName GameSpecificCategoryName = FApp::GetProjectName();
 	TSharedPtr<FSettingsCategory> AdvancedCategory;
+	TSharedPtr<FSettingsCategory> GameSpecificCategory;
 
 	for (TMap<FName, TSharedPtr<FSettingsCategory> >::TConstIterator It(Categories); It; ++It)
 	{
@@ -106,14 +108,24 @@ int32 FSettingsContainer::GetCategories( TArray<ISettingsCategoryPtr>& OutCatego
 			// Store off the advanced category, we'll add it to the bottom of all categories
 			AdvancedCategory = Category;
 		}
+		else if (Category->GetName() == GameSpecificCategoryName)
+		{
+			GameSpecificCategory = Category;
+		}
 		else
 		{
 			OutCategories.Add(It.Value());
 		}
 	}
 
+	// Always show the game specific category first
+	if (GameSpecificCategory.IsValid())
+	{
+		OutCategories.Insert(GameSpecificCategory, 0);
+	}
+
 	// always show the advanced category last
-	if(AdvancedCategory.IsValid())
+	if (AdvancedCategory.IsValid())
 	{
 		OutCategories.Add(AdvancedCategory);
 	}
