Add visibility controls to settings sections

diff --git a/Engine/Source/Developer/Settings/Private/SettingsCategory.cpp b/Engine/Source/Developer/Settings/Private/SettingsCategory.cpp
--- a/Engine/Source/Developer/Settings/Private/SettingsCategory.cpp
+++ b/Engine/Source/Developer/Settings/Private/SettingsCategory.cpp
@@ -57,18 +57,26 @@ void FSettingsCategory::RemoveSection( const FName& SectionName )
 /* ISettingsCategory interface
  *****************************************************************************/
 
-ISettingsSectionPtr FSettingsCategory::GetSection( const FName& SectionName ) const
+ISettingsSectionPtr FSettingsCategory::GetSection( const FName& SectionName, bool bIgnoreVisibility ) const
 {
-	return Sections.FindRef(SectionName);
+	ISettingsSectionPtr Section = Sections.FindRef(SectionName);
+	if (bIgnoreVisibility || (Section.IsValid() && SectionVisibilityPermissionList.PassesFilter(Section->GetName())))
+	{
+		return Section;
+	}
+	return nullptr;
 }
 
-int32 FSettingsCategory::GetSections( TArray<ISettingsSectionPtr>& OutSections ) const
+int32 FSettingsCategory::GetSections( TArray<ISettingsSectionPtr>& OutSections, bool bIgnoreVisibility ) const
 {
 	OutSections.Empty(Sections.Num());
 
 	for (TMap<FName, TSharedPtr<FSettingsSection> >::TConstIterator It(Sections); It; ++It)
 	{
-		OutSections.Add(It.Value());
+		if (bIgnoreVisibility || SectionVisibilityPermissionList.PassesFilter(It.Value()->GetName()))
+		{
+			OutSections.Add(It.Value());
+		}
 	}
 
 	return OutSections.Num();
diff --git a/Engine/Source/Developer/Settings/Private/SettingsCategory.h b/Engine/Source/Developer/Settings/Private/SettingsCategory.h
--- a/Engine/Source/Developer/Settings/Private/SettingsCategory.h
+++ b/Engine/Source/Developer/Settings/Private/SettingsCategory.h
@@ -5,6 +5,7 @@
 #include "CoreMinimal.h"
 #include "ISettingsCategory.h"
 #include "ISettingsSection.h"
+#include "Misc/NamePermissionList.h"
 
 class FSettingsSection;
 class SWidget;
@@ -89,9 +90,16 @@ public:
 		return Name;
 	}
 
-	virtual ISettingsSectionPtr GetSection( const FName& SectionName ) const override;
+	/** Gets a section if it's visible according to IsSectionVisiblePermissionList. bIgnoreVisibility=true will return the section even if it's filtered */
+	virtual ISettingsSectionPtr GetSection( const FName& SectionName, bool bIgnoreVisibility = false ) const override;
 
-	virtual int32 GetSections( TArray<ISettingsSectionPtr>& OutSections ) const override;
+	/** Gets all visible sections according to IsSectionVisiblePermissionList. bIgnoreVisibility=true will return sections even if they're filtered */
+	virtual int32 GetSections( TArray<ISettingsSectionPtr>& OutSections, bool bIgnoreVisibility = false ) const override;
+	
+	virtual FNamePermissionList* GetSectionVisibilityPermissionList() override
+	{
+		return &SectionVisibilityPermissionList;
+	}
 
 private:
 
@@ -106,4 +114,7 @@ private:
 
 	/** Holds the category's name. */
 	FName Name;
+
+	/** Determines which sections are returned with GetSection and GetSections */
+	FNamePermissionList SectionVisibilityPermissionList;
 };
diff --git a/Engine/Source/Developer/Settings/Private/SettingsContainer.cpp b/Engine/Source/Developer/Settings/Private/SettingsContainer.cpp
--- a/Engine/Source/Developer/Settings/Private/SettingsContainer.cpp
+++ b/Engine/Source/Developer/Settings/Private/SettingsContainer.cpp
@@ -55,7 +55,7 @@ void FSettingsContainer::RemoveSection( const FName& CategoryName, const FName&
 
 	if (Category.IsValid())
 	{
-		ISettingsSectionPtr Section = Category->GetSection(SectionName);
+		ISettingsSectionPtr Section = Category->GetSection(SectionName, true);
 
 		if (Section.IsValid())
 		{
diff --git a/Engine/Source/Developer/Settings/Public/ISettingsCategory.h b/Engine/Source/Developer/Settings/Public/ISettingsCategory.h
--- a/Engine/Source/Developer/Settings/Public/ISettingsCategory.h
+++ b/Engine/Source/Developer/Settings/Public/ISettingsCategory.h
@@ -37,7 +37,7 @@ public:
 	 * @param SectionName The name of the section to get.
 	 * @return The settings section, or nullptr if it doesn't exist.
 	 */
-	virtual TSharedPtr<ISettingsSection> GetSection( const FName& SectionName ) const = 0;
+	virtual TSharedPtr<ISettingsSection> GetSection( const FName& SectionName, bool bIgnoreVisibility = false ) const = 0;
 
 	/**
 	 * Gets the setting sections contained in this category.
@@ -45,7 +45,7 @@ public:
 	 * @param OutSections Will hold the collection of sections.
 	 * @return The number of sections returned.
 	 */
-	virtual int32 GetSections( TArray<TSharedPtr<ISettingsSection>>& OutSections ) const = 0;
+	virtual int32 GetSections( TArray<TSharedPtr<ISettingsSection>>& OutSections, bool bIgnoreVisibility = false ) const = 0;
 
 	/**
 	 * Gets the category's name.
@@ -54,6 +54,8 @@ public:
 	 */
 	virtual const FName& GetName() const = 0;
 
+	virtual class FNamePermissionList* GetSectionVisibilityPermissionList() = 0;
+
 public:
 
 	/** Virtual destructor. */
