Settings category sorting

diff --git a/Engine/Source/Developer/Settings/Private/SettingsContainer.cpp b/Engine/Source/Developer/Settings/Private/SettingsContainer.cpp
--- a/Engine/Source/Developer/Settings/Private/SettingsContainer.cpp
+++ b/Engine/Source/Developer/Settings/Private/SettingsContainer.cpp
@@ -3,6 +3,7 @@
 #include "SettingsContainer.h"
 #include "UObject/WeakObjectPtr.h"
 #include "Misc/App.h"
+#include "Algo/Transform.h"
 
 /* FSettingsContainer structors
  *****************************************************************************/
@@ -94,41 +95,45 @@ void FSettingsContainer::DescribeCategory( const FName& CategoryName, const FTex
 int32 FSettingsContainer::GetCategories( TArray<ISettingsCategoryPtr>& OutCategories ) const
 {
 	OutCategories.Empty(Categories.Num());
+	Algo::Transform(Categories, OutCategories, [](const TPair<FName, TSharedPtr<FSettingsCategory>>& Iter) { return Iter.Value; });
 
-	static const FName AdvancedCategoryName("Advanced");
-	const FName GameSpecificCategoryName = FApp::GetProjectName();
-	TSharedPtr<FSettingsCategory> AdvancedCategory;
-	TSharedPtr<FSettingsCategory> GameSpecificCategory;
-
-	for (TMap<FName, TSharedPtr<FSettingsCategory> >::TConstIterator It(Categories); It; ++It)
+	OutCategories.Sort([this](const ISettingsCategoryPtr& CatetoryA, const ISettingsCategoryPtr& CategoryB)
 	{
-		TSharedPtr<FSettingsCategory> Category = It.Value();
-		if(Category->GetName() == AdvancedCategoryName)
-		{
-			// Store off the advanced category, we'll add it to the bottom of all categories
-			AdvancedCategory = Category;
-		}
-		else if (Category->GetName() == GameSpecificCategoryName)
-		{
-			GameSpecificCategory = Category;
-		}
-		else
+		auto GetCategoryPriority = [this](const ISettingsCategoryPtr& Category)
 		{
-			OutCategories.Add(It.Value());
-		}
-	}
+			static const FName AdvancedCategoryName("Advanced");
+			const FName GameSpecificCategoryName = FApp::GetProjectName();
+
+			if (const float* Priority = CategorySortPriorities.Find(Category->GetName()))
+			{
+				return *Priority;
+			}
+			else if (Category->GetName() == GameSpecificCategoryName)
+			{
+				return -1.f;
+			}
+			else if (Category->GetName() == AdvancedCategoryName)
+			{
+				return 1.f;
+			}
+			else
+			{
+				return 0.f;
+			}
+		};
+
+		return GetCategoryPriority(CatetoryA) < GetCategoryPriority(CategoryB);
+	});
 
-	// Always show the game specific category first
-	if (GameSpecificCategory.IsValid())
-	{
-		OutCategories.Insert(GameSpecificCategory, 0);
-	}
+	return OutCategories.Num();
+}
 
-	// always show the advanced category last
-	if (AdvancedCategory.IsValid())
-	{
-		OutCategories.Add(AdvancedCategory);
-	}
+void FSettingsContainer::SetCategorySortPriority(FName CategoryName, float Priority)
+{
+	CategorySortPriorities.Add(CategoryName, Priority);
+}
 
-	return OutCategories.Num();
+void FSettingsContainer::ResetCategorySortPriority(FName CategoryName)
+{
+	CategorySortPriorities.Remove(CategoryName);
 }
diff --git a/Engine/Source/Developer/Settings/Private/SettingsContainer.h b/Engine/Source/Developer/Settings/Private/SettingsContainer.h
--- a/Engine/Source/Developer/Settings/Private/SettingsContainer.h
+++ b/Engine/Source/Developer/Settings/Private/SettingsContainer.h
@@ -104,6 +104,9 @@ public:
 		return SectionRemovedDelegate;
 	}
 
+	virtual void SetCategorySortPriority(FName CategoryName, float Priority) override;
+	virtual void ResetCategorySortPriority(FName CategoryName) override;
+
 private:
 
 	/** Holds the collection of setting categories. */
@@ -121,6 +124,9 @@ private:
 	/** Holds the container's name. */
 	FName Name;
 
+	/** Category sort priorities */
+	TMap<FName, float> CategorySortPriorities;
+
 private:
 
 	/** Holds a delegate that is executed when a settings category has been added or modified. */
diff --git a/Engine/Source/Developer/Settings/Public/ISettingsContainer.h b/Engine/Source/Developer/Settings/Public/ISettingsContainer.h
--- a/Engine/Source/Developer/Settings/Public/ISettingsContainer.h
+++ b/Engine/Source/Developer/Settings/Public/ISettingsContainer.h
@@ -79,6 +79,26 @@ public:
 	 */
 	virtual const FName& GetName() const = 0;
 
+	/**
+	 * Sets the sort priority of a category
+	 * A lower priority value means the category is listed first
+	 * Categories with the same priority are alphanumerically sorted
+	 * The priority value is 0 by default
+	 * @note A category named after FApp::GetProjectName() has a priority of -1.0 by default
+	 * @note A category named "Advanced" has a priority of 1.0 by default
+	 * @param ContainerName Settings container name
+	 * @param CategoryName Name of category to set the priority on
+	 * @param Priority Priority value
+	 */
+	virtual void SetCategorySortPriority(FName CategoryName, float Priority) = 0;
+
+	/**
+	 * Resets the sort priority of a category to its default value
+	 * @note Sort priority
+	 * @param CategoryName Name of category to set the priority on
+	 */
+	virtual void ResetCategorySortPriority(FName CategoryName) = 0;
+
 public:
 
 	/**
