Keep setttings category sort stable (update doc that pretends there's alphanumeric sorting going on)

diff --git a/Engine/Source/Developer/Settings/Private/SettingsContainer.cpp b/Engine/Source/Developer/Settings/Private/SettingsContainer.cpp
--- a/Engine/Source/Developer/Settings/Private/SettingsContainer.cpp
+++ b/Engine/Source/Developer/Settings/Private/SettingsContainer.cpp
@@ -97,7 +97,7 @@ int32 FSettingsContainer::GetCategories( TArray<ISettingsCategoryPtr>& OutCatego
 	OutCategories.Empty(Categories.Num());
 	Algo::Transform(Categories, OutCategories, [](const TPair<FName, TSharedPtr<FSettingsCategory>>& Iter) { return Iter.Value; });
 
-	OutCategories.Sort([this](const ISettingsCategoryPtr& CatetoryA, const ISettingsCategoryPtr& CategoryB)
+	OutCategories.StableSort([this](const ISettingsCategoryPtr& CatetoryA, const ISettingsCategoryPtr& CategoryB)
 	{
 		auto GetCategoryPriority = [this](const ISettingsCategoryPtr& Category)
 		{
diff --git a/Engine/Source/Developer/Settings/Public/ISettingsContainer.h b/Engine/Source/Developer/Settings/Public/ISettingsContainer.h
--- a/Engine/Source/Developer/Settings/Public/ISettingsContainer.h
+++ b/Engine/Source/Developer/Settings/Public/ISettingsContainer.h
@@ -82,7 +82,6 @@ public:
 	/**
 	 * Sets the sort priority of a category
 	 * A lower priority value means the category is listed first
-	 * Categories with the same priority are alphanumerically sorted
 	 * The priority value is 0 by default
 	 * @note A category named after FApp::GetProjectName() has a priority of -1.0 by default
 	 * @note A category named "Advanced" has a priority of 1.0 by default
