SettingsModule : Improve category descriptions (remove the "description text here" endings).

diff --git a/Engine/Source/Developer/Settings/Private/SettingsModule.cpp b/Engine/Source/Developer/Settings/Private/SettingsModule.cpp
--- a/Engine/Source/Developer/Settings/Private/SettingsModule.cpp
+++ b/Engine/Source/Developer/Settings/Private/SettingsModule.cpp
@@ -86,21 +86,21 @@ public:
 		// @todo gmp: move this into the modules that own these setting categories
 		TSharedRef<FSettingsContainer> EditorSettingsContainer = FindOrAddContainer("Editor");
 		EditorSettingsContainer->Describe(LOCTEXT("EditorPreferencesSubMenuLabel", "Editor Preferences"), LOCTEXT("EditorPreferencesSubMenuToolTip", "Configure the behavior and features of this Editor"), NAME_None);
-		EditorSettingsContainer->DescribeCategory("General", LOCTEXT("EditorGeneralCategoryName", "General"), LOCTEXT("EditorGeneralCategoryDescription", "General Editor settings description text here"));
-		EditorSettingsContainer->DescribeCategory("LevelEditor", LOCTEXT("EditorLevelEditorCategoryName", "Level Editor"), LOCTEXT("EditorLevelEditorCategoryDescription", "Level Editor settings description text here"));
-		EditorSettingsContainer->DescribeCategory("ContentEditors", LOCTEXT("EditorContentEditorsCategoryName", "Content Editors"), LOCTEXT("EditorContentEditorsCategoryDescription", "Content editors settings description text here"));
-		EditorSettingsContainer->DescribeCategory("Plugins", LOCTEXT("EditorPluginsCategoryName", "Plugins"), LOCTEXT("EditorPluginsCategoryDescription", "Plugins settings description text here"));
+		EditorSettingsContainer->DescribeCategory("General", LOCTEXT("EditorGeneralCategoryName", "General"), LOCTEXT("EditorGeneralCategoryDescription", "General Editor settings"));
+		EditorSettingsContainer->DescribeCategory("LevelEditor", LOCTEXT("EditorLevelEditorCategoryName", "Level Editor"), LOCTEXT("EditorLevelEditorCategoryDescription", "Level Editor settings"));
+		EditorSettingsContainer->DescribeCategory("ContentEditors", LOCTEXT("EditorContentEditorsCategoryName", "Content Editors"), LOCTEXT("EditorContentEditorsCategoryDescription", "Content editors settings"));
+		EditorSettingsContainer->DescribeCategory("Plugins", LOCTEXT("EditorPluginsCategoryName", "Plugins"), LOCTEXT("EditorPluginsCategoryDescription", "Plugins settings"));
 		EditorSettingsContainer->DescribeCategory("Advanced", LOCTEXT("EditorAdvancedCategoryName", "Advanced"), LOCTEXT("EditorAdvancedCategoryDescription", "Advanced editor settings"));
 
 		// @todo gmp: move this into the modules that own these setting categories
 		TSharedRef<FSettingsContainer> ProjectSettingsContainer = FindOrAddContainer("Project");
 		ProjectSettingsContainer->Describe(LOCTEXT("ProjectSettingsSubMenuLabel", "Project Settings"), LOCTEXT("ProjectSettingsSubMenuToolTip", "Change the settings of the currently loaded project"), NAME_None);
-		ProjectSettingsContainer->DescribeCategory("Project", LOCTEXT("ProjectProjectCategoryName", "Project"), LOCTEXT("ProjectProjectCategoryDescription", "Project settings description text here"));
-		ProjectSettingsContainer->DescribeCategory("Game", LOCTEXT("ProjectGameCategoryName", "Game"), LOCTEXT("ProjectGameCategoryDescription", "Game settings description text here"));
-		ProjectSettingsContainer->DescribeCategory("Engine", LOCTEXT("ProjectEngineCategoryName", "Engine"), LOCTEXT("ProjectEngineCategoryDescription", "Project settings description text here"));
+		ProjectSettingsContainer->DescribeCategory("Project", LOCTEXT("ProjectProjectCategoryName", "Project"), LOCTEXT("ProjectProjectCategoryDescription", "Project settings"));
+		ProjectSettingsContainer->DescribeCategory("Game", LOCTEXT("ProjectGameCategoryName", "Game"), LOCTEXT("ProjectGameCategoryDescription", "Game settings"));
+		ProjectSettingsContainer->DescribeCategory("Engine", LOCTEXT("ProjectEngineCategoryName", "Engine"), LOCTEXT("ProjectEngineCategoryDescription", "Project settings specific to the engine"));
 		ProjectSettingsContainer->DescribeCategory("Editor", LOCTEXT("ProjectEditorCategoryName", "Editor"), LOCTEXT("ProjectEditorCategoryDescription", "Project settings specific to the editor"));
-		ProjectSettingsContainer->DescribeCategory("Platforms", LOCTEXT("ProjectPlatformsCategoryName", "Platforms"), LOCTEXT("ProjectPlatformsCategoryDescription", "Platform settings description text here"));
-		ProjectSettingsContainer->DescribeCategory("Plugins", LOCTEXT("ProjectPluginsCategoryName", "Plugins"), LOCTEXT("ProjectPluginsCategoryDescription", "Plugins settings description text here"));
+		ProjectSettingsContainer->DescribeCategory("Platforms", LOCTEXT("ProjectPlatformsCategoryName", "Platforms"), LOCTEXT("ProjectPlatformsCategoryDescription", "Platform settings"));
+		ProjectSettingsContainer->DescribeCategory("Plugins", LOCTEXT("ProjectPluginsCategoryName", "Plugins"), LOCTEXT("ProjectPluginsCategoryDescription", "Plugins settings"));
 
 #if WITH_RELOAD
 		OnReloadHandle = FCoreUObjectDelegates::ReloadReinstancingCompleteDelegate.AddLambda([this] { ReinstancingComplete(); });
