Adding missing localization for the privacy category

#rnx

diff --git a/Engine/Source/Developer/Settings/Private/SettingsModule.cpp b/Engine/Source/Developer/Settings/Private/SettingsModule.cpp
--- a/Engine/Source/Developer/Settings/Private/SettingsModule.cpp
+++ b/Engine/Source/Developer/Settings/Private/SettingsModule.cpp
@@ -89,6 +89,7 @@ public:
 		EditorSettingsContainer->DescribeCategory("General", LOCTEXT("EditorGeneralCategoryName", "General"), LOCTEXT("EditorGeneralCategoryDescription", "General Editor settings"));
 		EditorSettingsContainer->DescribeCategory("LevelEditor", LOCTEXT("EditorLevelEditorCategoryName", "Level Editor"), LOCTEXT("EditorLevelEditorCategoryDescription", "Level Editor settings"));
 		EditorSettingsContainer->DescribeCategory("ContentEditors", LOCTEXT("EditorContentEditorsCategoryName", "Content Editors"), LOCTEXT("EditorContentEditorsCategoryDescription", "Content editors settings"));
+		EditorSettingsContainer->DescribeCategory("Privacy", LOCTEXT("EditorPrivacyCategoryName", "Privacy"), LOCTEXT("EditorPrivacyCategoryDescription", "Privacy settings"));
 		EditorSettingsContainer->DescribeCategory("Plugins", LOCTEXT("EditorPluginsCategoryName", "Plugins"), LOCTEXT("EditorPluginsCategoryDescription", "Plugins settings"));
 		EditorSettingsContainer->DescribeCategory("Advanced", LOCTEXT("EditorAdvancedCategoryName", "Advanced"), LOCTEXT("EditorAdvancedCategoryDescription", "Advanced editor settings"));
 
