Editor Preferences: Reset to Defaults should now function regardless of whether Set as Default was ever used.

[RN]

diff --git a/Engine/Source/Developer/Settings/Private/SettingsSection.cpp b/Engine/Source/Developer/Settings/Private/SettingsSection.cpp
--- a/Engine/Source/Developer/Settings/Private/SettingsSection.cpp
+++ b/Engine/Source/Developer/Settings/Private/SettingsSection.cpp
@@ -5,7 +5,9 @@
 #include "Misc/ConfigCacheIni.h"
 #include "Misc/ConfigContext.h"
 #include "UObject/Class.h"
+#include "UObject/Package.h"
 #include "UObject/Reload.h"
+#include "UObject/UnrealType.h"
 
 
 /* FSettingsSection structors
@@ -173,7 +175,6 @@ bool FSettingsSection::Import( const FString& Filename )
 	return false;	
 }
 
-
 bool FSettingsSection::ResetDefaults()
 {
 	if (ResetDefaultsDelegate.IsBound())
@@ -183,13 +184,60 @@ bool FSettingsSection::ResetDefaults()
 
 	if (SettingsObject.IsValid() && SettingsObject->GetClass()->HasAnyClassFlags(CLASS_Config) && !SettingsObject->GetClass()->HasAnyClassFlags(CLASS_DefaultConfig | CLASS_GlobalUserConfig | CLASS_ProjectUserConfig))
 	{
-		FString ConfigName = SettingsObject->GetClass()->GetConfigName();
+		UClass* SettingsClass = SettingsObject->GetClass();
+		const FString ConfigName = SettingsClass->GetConfigName();
 
-		GConfig->EmptySection(*SettingsObject->GetClass()->GetPathName(), ConfigName);
+		// Remove the section from the standard config
+		// Ex. If resetting AudioEditorSettings, this removes the "/Script/AudioEditor.AudioEditorSettings" section from Engine/Saved/Config/WindowsEditor/EditorSettings.ini
+		GConfig->EmptySection(*SettingsClass->GetPathName(), ConfigName);
 		GConfig->Flush(false);
 
+		// Reload the same config, leaving the cache with whatever defaults remain for those same values
 		FConfigContext::ForceReloadIntoGConfig().Load(*FPaths::GetBaseFilename(ConfigName));
 
+		// Create a temporary default object to get default property values from
+		UObject* TempDefaultObject;
+		{
+			UClass* ParentClass = SettingsClass->GetSuperClass();
+			UObject* ParentDefaultObject = nullptr;
+			if (ParentClass)
+			{
+				UObjectForceRegistration(ParentClass);
+				ParentDefaultObject = ParentClass->GetDefaultObject();
+				check(GConfig);
+				if (GEventDrivenLoaderEnabled && EVENT_DRIVEN_ASYNC_LOAD_ACTIVE_AT_RUNTIME)
+				{
+					check(ParentDefaultObject && !ParentDefaultObject->HasAnyFlags(RF_NeedLoad));
+				}
+			}
+			TempDefaultObject = StaticAllocateObject(SettingsClass, GetTransientPackage(), NAME_None, RF_ClassDefaultObject);
+			check(SettingsObject.Get() != TempDefaultObject);
+
+			EObjectInitializerOptions InitOptions = EObjectInitializerOptions::None;
+			if (!SettingsClass->HasAnyClassFlags(CLASS_Native | CLASS_Intrinsic))
+			{
+				// Blueprint CDOs have their properties always initialized.
+				InitOptions |= EObjectInitializerOptions::InitializeProperties;
+			}
+			(*SettingsClass->ClassConstructor)(FObjectInitializer(TempDefaultObject, ParentDefaultObject, InitOptions));
+			TempDefaultObject->Rename(nullptr, nullptr, REN_DoNotDirty | REN_DontCreateRedirectors | REN_NonTransactional);
+		}
+
+		// Copy config property values from the temporary default object back into the actual CDO
+		for (const FProperty* Property = SettingsClass->PropertyLink; Property; Property = Property->PropertyLinkNext)
+		{
+			if (!Property->HasAnyPropertyFlags(CPF_Config | CPF_GlobalConfig))
+			{
+				continue;
+			}
+
+			Property->CopyCompleteValue_InContainer(SettingsObject.Get(), TempDefaultObject);
+		}
+
+		TempDefaultObject->MarkAsGarbage();
+		TempDefaultObject = nullptr;
+
+		// Propagate the changes system-wide
 		SettingsObject->ReloadConfig(nullptr, nullptr, UE::LCPF_PropagateToInstances|UE::LCPF_PropagateToChildDefaultObjects);
 
 		return true;
diff --git a/Engine/Source/Runtime/Core/Private/Misc/ConfigContext.cpp b/Engine/Source/Runtime/Core/Private/Misc/ConfigContext.cpp
--- a/Engine/Source/Runtime/Core/Private/Misc/ConfigContext.cpp
+++ b/Engine/Source/Runtime/Core/Private/Misc/ConfigContext.cpp
@@ -432,7 +432,17 @@ bool FConfigContext::PrepareForLoad(bool& bPerformLoad)
 		// FPaths::GeneratedConfigDir() will permanently cache the value
 		if (GeneratedConfigDir.IsEmpty())
 		{
-			GeneratedConfigDir = FPaths::GeneratedConfigDir();
+			// Use the engine directory for editor settings
+			if (BaseIniName == TEXT("EditorSettings") ||
+				BaseIniName == TEXT("EditorKeyBindings") ||
+				BaseIniName == TEXT("EditorLayout"))
+			{
+				GeneratedConfigDir = FPaths::EngineEditorSettingsDir();
+			}
+			else
+			{
+				GeneratedConfigDir = FPaths::GeneratedConfigDir();
+			}
 		}
 
 		// calculate where this file will be saved/generated to (or at least the key to look up in the ConfigSystem)
